Media: Initial revision of media assets

diff --git a/Engine/Config/BaseEditor.ini b/Engine/Config/BaseEditor.ini
--- a/Engine/Config/BaseEditor.ini
+++ b/Engine/Config/BaseEditor.ini
@@ -34,6 +34,7 @@ MaxNumberOfLogLines=2000
 +RenderableThumbnailTypes=(ClassNeedingThumbnailName="/Script/Engine.World",RendererClassName="/Script/UnrealEd.WorldThumbnailRenderer")
 +RenderableThumbnailTypes=(ClassNeedingThumbnailName="/Script/Engine.SoundWave",RendererClassName="/Script/UnrealEd.SoundWaveThumbnailRenderer")
 +RenderableThumbnailTypes=(ClassNeedingThumbnailName="/Script/Engine.SlateBrushAsset",RendererClassName="/Script/UnrealEd.SlateBrushThumbnailRenderer")
++RenderableThumbnailTypes=(ClassNeedingThumbnailName="/Script/MediaAssets.MediaTexture",RendererClassName="/Script/UnrealEd.TextureThumbnailRenderer")
 
 [/Script/UnrealEd.DefaultSizedThumbnailRenderer]
 DefaultSizeX=512
diff --git a/Engine/Source/Runtime/MediaAssets/MediaAssets.Build.cs b/Engine/Source/Runtime/MediaAssets/MediaAssets.Build.cs
new file mode 100644
--- /dev/null
+++ b/Engine/Source/Runtime/MediaAssets/MediaAssets.Build.cs
@@ -0,0 +1,33 @@
+// Copyright 1998-2014 Epic Games, Inc. All Rights Reserved.
+
+namespace UnrealBuildTool.Rules
+{
+	public class MediaAssets : ModuleRules
+	{
+		public MediaAssets(TargetInfo Target)
+		{
+            PublicDependencyModuleNames.AddRange(
+                new string[] {
+					"Core",
+					"CoreUObject",
+					"Engine",                
+                }
+            );
+
+			PrivateDependencyModuleNames.AddRange(
+				new string[] {
+					"Media",
+                    "RenderCore",
+                    "RHI",
+				}
+			); 
+
+			PrivateIncludePaths.AddRange(
+				new string[] {
+					"Runtime/MediaAssets/Private",
+					"Runtime/MediaAssets/Private/Assets",
+				}
+			);
+		}
+	}
+}
diff --git a/Engine/Source/Runtime/MediaAssets/Private/Assets/MaterialExpressionTextureSampleParameterMedia.cpp b/Engine/Source/Runtime/MediaAssets/Private/Assets/MaterialExpressionTextureSampleParameterMedia.cpp
new file mode 100644
--- /dev/null
+++ b/Engine/Source/Runtime/MediaAssets/Private/Assets/MaterialExpressionTextureSampleParameterMedia.cpp
@@ -0,0 +1,66 @@
+// Copyright 1998-2014 Epic Games, Inc. All Rights Reserved.
+
+#include "MediaAssetsPrivatePCH.h"
+
+
+#define LOCTEXT_NAMESPACE "UMaterialExpressionTextureSampleParameterMedia"
+
+
+/* UMaterialExpressionTextureSampleParameterMedia structors
+ *****************************************************************************/
+
+UMaterialExpressionTextureSampleParameterMedia::UMaterialExpressionTextureSampleParameterMedia( const class FPostConstructInitializeProperties& PCIP )
+	: Super(PCIP)
+{
+	// Structure to hold one-time initialization
+	struct FConstructorStatics
+	{
+		FString NAME_Texture;
+		FString NAME_Parameters;
+		FConstructorStatics()
+			: NAME_Texture(LOCTEXT("Texture", "Texture").ToString())
+			, NAME_Parameters(LOCTEXT("Parameters", "Parameters").ToString())
+		{ }
+	};
+
+	static FConstructorStatics ConstructorStatics;
+
+	MenuCategories.Empty();
+	MenuCategories.Add(ConstructorStatics.NAME_Texture);
+	MenuCategories.Add(ConstructorStatics.NAME_Parameters);
+}
+
+
+/* UMaterialExpressionTextureSampleParameter interface
+ *****************************************************************************/
+
+const TCHAR* UMaterialExpressionTextureSampleParameterMedia::GetRequirements( )
+{
+	return TEXT("Requires MediaTexture");
+}
+
+
+bool UMaterialExpressionTextureSampleParameterMedia::TextureIsValid( UTexture* InTexture )
+{
+	bool Result = false;
+
+	if (InTexture)
+	{
+		Result = (InTexture->GetClass() == UMediaTexture::StaticClass());
+	}
+
+	return Result;
+}
+
+
+/* UMaterialExpression interface
+ *****************************************************************************/
+
+void UMaterialExpressionTextureSampleParameterMedia::GetCaption( TArray<FString>& OutCaptions ) const
+{
+	OutCaptions.Add(TEXT("ParamMedia")); 
+	OutCaptions.Add(FString::Printf(TEXT("'%s'"), *ParameterName.ToString()));
+}
+
+
+#undef LOCTEXT_NAMESPACE
diff --git a/Engine/Source/Runtime/MediaAssets/Private/Assets/MediaAsset.cpp b/Engine/Source/Runtime/MediaAssets/Private/Assets/MediaAsset.cpp
new file mode 100644
--- /dev/null
+++ b/Engine/Source/Runtime/MediaAssets/Private/Assets/MediaAsset.cpp
@@ -0,0 +1,302 @@
+// Copyright 1998-2014 Epic Games, Inc. All Rights Reserved.
+
+#include "MediaAssetsPrivatePCH.h"
+
+
+/* UMediaAsset structors
+ *****************************************************************************/
+
+UMediaAsset::UMediaAsset( const class FPostConstructInitializeProperties& PCIP )
+	: Super(PCIP)
+	, AutoPlay(false)
+	, Looping(true)
+	, PlaybackRate(1.0f)
+	, ResetOnLastFrame(false)
+	, StreamMode(MASM_FromMemory)
+	, MediaPlayer(nullptr)
+{ }
+
+
+/* UMediaAsset interface
+ *****************************************************************************/
+
+bool UMediaAsset::CanPause( ) const
+{
+	return MediaPlayer.IsValid() && MediaPlayer->IsPlaying();
+}
+
+
+bool UMediaAsset::CanPlay( ) const
+{
+	return MediaPlayer.IsValid() && !MediaPlayer->IsPlaying() && MediaPlayer->IsReady();
+}
+
+
+bool UMediaAsset::CanStop( ) const
+{
+	return MediaPlayer.IsValid() && MediaPlayer->IsReady();
+}
+
+
+FTimespan UMediaAsset::GetDuration( ) const
+{
+	if (MediaPlayer.IsValid())
+	{
+		return MediaPlayer->GetDuration();
+	}
+
+	return FTimespan::Zero();
+}
+
+
+float UMediaAsset::GetRate( ) const
+{
+	if (MediaPlayer.IsValid())
+	{
+		return MediaPlayer->GetRate();
+	}
+
+	return 0;
+}
+
+
+FTimespan UMediaAsset::GetTime( ) const
+{
+	if (MediaPlayer.IsValid())
+	{
+		return MediaPlayer->GetTime();
+	}
+
+	return FTimespan::Zero();
+}
+
+
+bool UMediaAsset::IsPaused( ) const
+{
+	return MediaPlayer.IsValid() && MediaPlayer->IsPaused();
+}
+
+
+bool UMediaAsset::IsPlaying( ) const
+{
+	return MediaPlayer.IsValid() && MediaPlayer->IsPlaying();
+}
+
+
+bool UMediaAsset::IsStopped( ) const
+{
+	return !MediaPlayer.IsValid() || !MediaPlayer->IsReady();
+}
+
+
+bool UMediaAsset::Pause( )
+{
+	return MediaPlayer.IsValid() && MediaPlayer->Pause();
+}
+
+
+bool UMediaAsset::Rewind( )
+{
+	return MediaPlayer.IsValid() && MediaPlayer->Seek(FTimespan::Zero());
+}
+
+
+bool UMediaAsset::Play( )
+{
+	return MediaPlayer.IsValid() && MediaPlayer->Play();
+}
+
+
+bool UMediaAsset::SetRate( float InRate )
+{
+	PlaybackRate = InRate;
+
+	if (!MediaPlayer.IsValid() || !MediaPlayer->IsPlaying())
+	{
+		return false;
+	}
+
+	return MediaPlayer->Play(InRate);
+}
+
+
+bool UMediaAsset::Seek( const FTimespan& InTime )
+{
+	return MediaPlayer.IsValid() && MediaPlayer->Seek(InTime);
+}
+
+
+void UMediaAsset::Stop( )
+{
+	if (MediaPlayer.IsValid())
+	{
+		Rewind();
+		Pause();
+	}
+}
+
+
+bool UMediaAsset::SupportsRate( float Rate, bool Unthinned ) const
+{
+	return MediaPlayer.IsValid() && MediaPlayer->SupportsRate(Rate, Unthinned);
+}
+
+
+bool UMediaAsset::SupportsScrubbing( ) const
+{
+	return MediaPlayer.IsValid() && MediaPlayer->SupportsScrubbing();
+}
+
+
+bool UMediaAsset::SupportsSeeking( ) const
+{
+	return MediaPlayer.IsValid() && MediaPlayer->SupportsSeeking();
+}
+
+
+/* UObject  overrides
+ *****************************************************************************/
+
+void UMediaAsset::BeginDestroy( )
+{
+	Super::BeginDestroy();
+
+	if (MediaPlayer.IsValid())
+	{
+		MediaPlayer->Close();
+		MediaPlayer.Reset();
+	}
+}
+
+
+FString UMediaAsset::GetDesc( )
+{
+	if (MediaPlayer.IsValid())
+	{
+		return TEXT("UMediaAsset");
+	}
+
+	return FString();
+}
+
+
+void UMediaAsset::PostLoad( )
+{
+	Super::PostLoad();
+
+	if (!HasAnyFlags(RF_ClassDefaultObject) && !GIsBuildMachine)
+	{
+		InitializePlayer();
+	}
+}
+
+
+#if WITH_EDITOR
+
+void UMediaAsset::PostEditChangeProperty( FPropertyChangedEvent& PropertyChangedEvent )
+{
+	Super::PostEditChangeProperty(PropertyChangedEvent);
+
+	InitializePlayer();
+}
+
+#endif
+
+
+/* UMediaAsset implementation
+ *****************************************************************************/
+
+void UMediaAsset::InitializePlayer( )
+{
+	const FString& NewUrl = URL.FilePath;
+
+	if (NewUrl != CurrentUrl)
+	{
+		// close previous player
+		CurrentUrl = FString();
+
+		if (MediaPlayer.IsValid())
+		{
+			MediaPlayer->Close();
+			MediaPlayer->OnClosing().RemoveAll(this);
+			MediaPlayer->OnOpened().RemoveAll(this);
+			MediaPlayer.Reset();
+		}
+
+		if (NewUrl.IsEmpty())
+		{
+			return;
+		}
+
+		// create new player
+		MediaPlayer = IMediaModule::Get().CreatePlayer(NewUrl);
+
+		if (!MediaPlayer.IsValid())
+		{
+			return;
+		}
+
+		MediaPlayer->OnClosing().AddUObject(this, &UMediaAsset::HandleMediaPlayerMediaClosing);
+		MediaPlayer->OnOpened().AddUObject(this, &UMediaAsset::HandleMediaPlayerMediaOpened);
+
+		// open the new media file
+		bool OpenedSuccessfully = false;
+
+		if (StreamMode == EMediaAssetStreamModes::MASM_FromUrl)
+		{
+			OpenedSuccessfully = MediaPlayer->Open(NewUrl);
+		}
+		else
+		{
+			FArchive* FileReader = IFileManager::Get().CreateFileReader(*NewUrl);
+		
+			if (FileReader == nullptr)
+			{
+				return;
+			}
+
+			if (FileReader->TotalSize() > 0)
+			{
+				TArray<uint8>* FileData = new TArray<uint8>();
+
+				FileData->AddUninitialized(FileReader->TotalSize());
+				FileReader->Serialize(FileData->GetData(), FileReader->TotalSize());
+
+				OpenedSuccessfully = MediaPlayer->Open(MakeShareable(FileData), NewUrl);
+			}
+
+			delete FileReader;
+		}
+
+		// finish initialization
+		if (OpenedSuccessfully)
+		{
+			CurrentUrl = NewUrl;
+		}
+	}
+
+	if (!MediaPlayer.IsValid())
+	{
+		return;
+	}
+    
+	if (AutoPlay)
+	{
+		MediaPlayer->Play();
+	}
+}
+
+
+/* UMediaAsset callbacks
+ *****************************************************************************/
+
+void UMediaAsset::HandleMediaPlayerMediaClosing( FString ClosedUrl )
+{
+	OnMediaClosing.Broadcast(ClosedUrl);
+}
+
+
+void UMediaAsset::HandleMediaPlayerMediaOpened( FString OpenedUrl )
+{
+	OnMediaOpened.Broadcast(OpenedUrl);
+}
diff --git a/Engine/Source/Runtime/MediaAssets/Private/Assets/MediaSoundWave.cpp b/Engine/Source/Runtime/MediaAssets/Private/Assets/MediaSoundWave.cpp
new file mode 100644
--- /dev/null
+++ b/Engine/Source/Runtime/MediaAssets/Private/Assets/MediaSoundWave.cpp
@@ -0,0 +1,210 @@
+// Copyright 1998-2014 Epic Games, Inc. All Rights Reserved.
+
+#include "MediaAssetsPrivatePCH.h"
+
+
+/* UMediaSoundWave structors
+ *****************************************************************************/
+
+UMediaSoundWave::UMediaSoundWave( const class FPostConstructInitializeProperties& PCIP )
+	: Super(PCIP)
+	, AudioQueue(MakeShareable(new FMediaSampleQueue))
+{
+	bLooping = false;
+	bProcedural = true;
+	Duration = INDEFINITELY_LOOPING_DURATION;
+}
+
+
+UMediaSoundWave::~UMediaSoundWave( )
+{
+	if (AudioTrack.IsValid())
+	{
+		AudioTrack->RemoveSink(AudioQueue);
+	}
+}
+
+
+/* UMediaSoundWave interface
+ *****************************************************************************/
+
+TSharedPtr<IMediaPlayer> UMediaSoundWave::GetMediaPlayer( ) const
+{
+	if (MediaAsset == nullptr)
+	{
+		return nullptr;
+	}
+
+	return MediaAsset->GetMediaPlayer();
+}
+
+
+void UMediaSoundWave::SetMediaAsset( UMediaAsset* InMediaAsset )
+{
+	MediaAsset = InMediaAsset;
+
+	InitializeStream();
+}
+
+
+/* USoundWave overrides
+ *****************************************************************************/
+
+int32 UMediaSoundWave::GeneratePCMData( uint8* PCMData, const int32 SamplesNeeded )
+{
+	// drain media sample queue
+	TSharedPtr<TArray<uint8>, ESPMode::ThreadSafe> Sample;
+
+	while (AudioQueue->Dequeue(Sample))
+	{
+		QueuedAudio.Append(*Sample);
+	}
+
+	// get requested samples
+	if (SamplesNeeded <= 0)
+	{
+		return 0;
+	}
+
+	const int32 SamplesAvailable = QueuedAudio.Num() / sizeof(int16);
+	const int32 BytesToCopy = FMath::Min<int32>(SamplesNeeded, SamplesAvailable) * sizeof(int16);
+
+	if (BytesToCopy > 0)
+	{
+		FMemory::Memcpy((void*)PCMData, &QueuedAudio[0], BytesToCopy);
+		QueuedAudio.RemoveAt(0, BytesToCopy);
+	}
+
+	return BytesToCopy;
+}
+
+
+FByteBulkData* UMediaSoundWave::GetCompressedData( FName Format )
+{
+	return nullptr;
+}
+
+
+int32 UMediaSoundWave::GetResourceSizeForFormat( FName Format )
+{
+	return 0;
+}
+
+
+void UMediaSoundWave::InitAudioResource( FByteBulkData& CompressedData )
+{
+	check(false); // should never be pushing compressed data to this class
+}
+
+
+bool UMediaSoundWave::InitAudioResource( FName Format )
+{
+	return true;
+}
+
+
+/* UObject overrides
+ *****************************************************************************/
+
+void UMediaSoundWave::GetAssetRegistryTags( TArray<FAssetRegistryTag>& OutTags ) const
+{
+}
+
+
+SIZE_T UMediaSoundWave::GetResourceSize(EResourceSizeMode::Type Mode)
+{
+	return 0;
+}
+
+
+void UMediaSoundWave::Serialize( FArchive& Ar )
+{
+	// do not call the USoundWave version of serialize
+	USoundBase::Serialize(Ar);
+}
+
+
+/* UMediaSoundWave implementation
+ *****************************************************************************/
+
+void UMediaSoundWave::InitializeStream( )
+{
+	// assign new media asset
+	if (CurrentMediaAsset != MediaAsset)
+	{
+		if (CurrentMediaAsset != nullptr)
+		{
+			IMediaPlayerPtr MediaPlayer = CurrentMediaAsset->GetMediaPlayer();
+
+			if (MediaPlayer.IsValid())
+			{
+				MediaPlayer->OnClosing().RemoveAll(this);
+				MediaPlayer->OnOpened().RemoveAll(this);			
+			}
+		}
+
+		CurrentMediaAsset = MediaAsset;
+
+		if (MediaAsset != nullptr)
+		{
+			IMediaPlayerPtr MediaPlayer = MediaAsset->GetMediaPlayer();
+
+			if (MediaPlayer.IsValid())
+			{
+				MediaPlayer->OnClosing().AddUObject(this, &UMediaSoundWave::HandleMediaPlayerClosing);
+				MediaPlayer->OnOpened().AddUObject(this, &UMediaSoundWave::HandleMediaPlayerOpened);
+			}
+		}	
+	}
+
+	// disconnect from current track
+	if (AudioTrack.IsValid())
+	{
+		AudioTrack->RemoveSink(AudioQueue);
+	}
+
+	AudioTrack.Reset();
+
+	// initialize from new track
+	if (MediaAsset != nullptr)
+	{
+		IMediaPlayerPtr MediaPlayer = MediaAsset->GetMediaPlayer();
+
+		if (MediaPlayer.IsValid())
+		{
+			AudioTrack = MediaPlayer->GetTrackSafe(AudioTrackIndex, EMediaTrackTypes::Audio);
+		}
+	}
+
+	if (AudioTrack.IsValid())
+	{
+		NumChannels = AudioTrack->GetAudioDetails().GetNumChannels();
+		SampleRate = AudioTrack->GetAudioDetails().GetSamplesPerSecond();
+	}
+	else
+	{
+		NumChannels = 0;
+		SampleRate = 0;
+	}
+
+	// connect to new track
+	if (AudioTrack.IsValid())
+	{
+		AudioTrack->AddSink(AudioQueue);
+	}
+}
+
+
+/* UMediaSoundWave callbacks
+ *****************************************************************************/
+
+void UMediaSoundWave::HandleMediaPlayerClosing( FString ClosingUrl )
+{
+	InitializeStream();
+}
+
+
+void UMediaSoundWave::HandleMediaPlayerOpened( FString OpenedUrl )
+{
+	InitializeStream();
+}
diff --git a/Engine/Source/Runtime/MediaAssets/Private/Assets/MediaTexture.cpp b/Engine/Source/Runtime/MediaAssets/Private/Assets/MediaTexture.cpp
new file mode 100644
--- /dev/null
+++ b/Engine/Source/Runtime/MediaAssets/Private/Assets/MediaTexture.cpp
@@ -0,0 +1,239 @@
+// Copyright 1998-2014 Epic Games, Inc. All Rights Reserved.
+
+#include "MediaAssetsPrivatePCH.h"
+#include "MediaTextureResource.h"
+
+
+/* UMediaTexture structors
+ *****************************************************************************/
+
+UMediaTexture::UMediaTexture( const class FPostConstructInitializeProperties& PCIP )
+	: Super(PCIP)
+	, MediaAsset(nullptr)
+	, ClearColor(FLinearColor::Red)
+	, CurrentMediaAsset(nullptr)
+	, VideoBuffer(MakeShareable(new FMediaSampleBuffer))
+{
+	NeverStream = true;
+
+	UpdateResource();
+}
+
+
+UMediaTexture::~UMediaTexture( )
+{
+	if (VideoTrack.IsValid())
+	{
+		VideoTrack->RemoveSink(VideoBuffer);
+	}
+}
+
+
+/* UMediaTexture interface
+ *****************************************************************************/
+
+void UMediaTexture::SetMediaAsset( UMediaAsset* InMediaAsset )
+{
+	MediaAsset = InMediaAsset;
+
+	InitializeTrack();
+}
+
+
+/* UTexture  overrides
+ *****************************************************************************/
+
+FTextureResource* UMediaTexture::CreateResource( )
+{
+	return new FMediaTextureResource(this, VideoBuffer);
+}
+
+
+EMaterialValueType UMediaTexture::GetMaterialType( )
+{
+	return MCT_Texture2D;
+}
+
+
+float UMediaTexture::GetSurfaceWidth( ) const
+{
+	return CachedDimensions.X;
+}
+
+
+float UMediaTexture::GetSurfaceHeight( ) const
+{
+	return CachedDimensions.Y;
+}
+
+
+/* UObject  overrides
+ *****************************************************************************/
+
+void UMediaTexture::BeginDestroy( )
+{
+	Super::BeginDestroy();
+
+	// synchronize with the rendering thread by inserting a fence
+ 	if (!ReleasePlayerFence)
+ 	{
+ 		ReleasePlayerFence = new FRenderCommandFence();
+ 	}
+
+ 	ReleasePlayerFence->BeginFence();
+}
+
+
+void UMediaTexture::FinishDestroy( )
+{
+	delete ReleasePlayerFence;
+	ReleasePlayerFence = nullptr;
+
+	Super::FinishDestroy();
+}
+
+
+FString UMediaTexture::GetDesc( )
+{
+	TSharedPtr<IMediaPlayer> MediaPlayer = GetMediaPlayer();
+
+	if (!MediaPlayer.IsValid())
+	{
+		return FString();
+	}
+
+	return FString::Printf(TEXT("%dx%d [%s]"), CachedDimensions.X, CachedDimensions.Y, GPixelFormats[GetFormat()].Name);
+}
+
+
+SIZE_T UMediaTexture::GetResourceSize( EResourceSizeMode::Type Mode )
+{
+	return CachedDimensions.X * CachedDimensions.Y * 4;
+}
+
+
+bool UMediaTexture::IsReadyForFinishDestroy( )
+{
+	// ready to call FinishDestroy if the flushing fence has been hit
+	return (Super::IsReadyForFinishDestroy() && ReleasePlayerFence && ReleasePlayerFence->IsFenceComplete());
+}
+
+
+void UMediaTexture::PostLoad( )
+{
+	Super::PostLoad();
+
+	if (!HasAnyFlags(RF_ClassDefaultObject) && !GIsBuildMachine)
+	{
+		InitializeTrack();
+	}
+}
+
+
+#if WITH_EDITOR
+
+void UMediaTexture::PreEditChange( UProperty* PropertyAboutToChange )
+{
+	// this will release the FMediaTextureResource
+	Super::PreEditChange(PropertyAboutToChange);
+
+	FlushRenderingCommands();
+}
+
+
+void UMediaTexture::PostEditChangeProperty( FPropertyChangedEvent& PropertyChangedEvent )
+{
+	InitializeTrack();
+
+	// this will recreate the FMediaTextureResource
+	Super::PostEditChangeProperty(PropertyChangedEvent);
+}
+
+
+/* UMediaTexture implementation
+ *****************************************************************************/
+
+void UMediaTexture::InitializeTrack( )
+{
+	// assign new media asset
+	if (CurrentMediaAsset != MediaAsset)
+	{
+		if (CurrentMediaAsset != nullptr)
+		{
+			IMediaPlayerPtr MediaPlayer = CurrentMediaAsset->GetMediaPlayer();
+
+			if (MediaPlayer.IsValid())
+			{
+				MediaPlayer->OnClosing().RemoveAll(this);
+				MediaPlayer->OnOpened().RemoveAll(this);
+			}
+		}
+
+		CurrentMediaAsset = MediaAsset;
+
+		if (MediaAsset != nullptr)
+		{
+			IMediaPlayerPtr MediaPlayer = MediaAsset->GetMediaPlayer();
+
+			if (MediaPlayer.IsValid())
+			{
+				MediaPlayer->OnClosing().AddUObject(this, &UMediaTexture::HandleMediaPlayerClosing);
+				MediaPlayer->OnOpened().AddUObject(this, &UMediaTexture::HandleMediaPlayerOpened);
+			}
+		}	
+	}
+
+	// disconnect from current track
+	if (VideoTrack.IsValid())
+	{
+		VideoTrack->RemoveSink(VideoBuffer);
+	}
+
+	VideoTrack.Reset();
+
+	// initialize from new track
+	if (MediaAsset != nullptr)
+	{
+		IMediaPlayerPtr MediaPlayer = MediaAsset->GetMediaPlayer();
+
+		if (MediaPlayer.IsValid())
+		{
+			VideoTrack = MediaPlayer->GetTrackSafe(VideoTrackIndex, EMediaTrackTypes::Video);
+		}
+	}
+
+	if (VideoTrack.IsValid())
+	{
+		CachedDimensions = VideoTrack->GetVideoDetails().GetDimensions();
+	}
+	else
+	{
+		CachedDimensions = FIntPoint(ForceInit);
+	}
+
+	UpdateResource();
+
+	// connect to new track
+	if (VideoTrack.IsValid())
+	{
+		VideoTrack->AddSink(VideoBuffer);
+	}
+}
+
+
+/* UMediaTexture callbacks
+ *****************************************************************************/
+
+void UMediaTexture::HandleMediaPlayerClosing( FString ClosingUrl )
+{
+	InitializeTrack();
+}
+
+
+void UMediaTexture::HandleMediaPlayerOpened( FString OpenedUrl )
+{
+	InitializeTrack();
+}
+
+
+#endif // WITH_EDITOR
diff --git a/Engine/Source/Runtime/MediaAssets/Private/Assets/MediaTextureResource.cpp b/Engine/Source/Runtime/MediaAssets/Private/Assets/MediaTextureResource.cpp
new file mode 100644
--- /dev/null
+++ b/Engine/Source/Runtime/MediaAssets/Private/Assets/MediaTextureResource.cpp
@@ -0,0 +1,121 @@
+// Copyright 1998-2014 Epic Games, Inc. All Rights Reserved.
+
+#include "MediaAssetsPrivatePCH.h"
+
+
+/* FMediaTextureResource structors
+ *****************************************************************************/
+
+FMediaTextureResource::FMediaTextureResource( const class UMediaTexture* InOwner, const FMediaSampleBufferRef& InVideoBuffer )
+	: Owner(InOwner)
+	, Cleared(false)
+	, LastFrameTime(FTimespan::MinValue())
+	, VideoBuffer(InVideoBuffer)
+{ }
+
+
+/* FTextureResource overrides
+ *****************************************************************************/
+
+void FMediaTextureResource::InitDynamicRHI( )
+{
+	if ((Owner->GetSurfaceWidth() > 0) && (Owner->GetSurfaceHeight() > 0))
+	{
+		// Create the RHI texture. Only one mip is used and the texture is targetable or resolve.
+		uint32 TexCreateFlags = Owner->SRGB ? TexCreate_SRGB : 0;
+		FRHIResourceCreateInfo CreateInfo;
+
+		RHICreateTargetableShaderResource2D(
+			Owner->GetSurfaceWidth(),
+			Owner->GetSurfaceHeight(),
+			Owner->GetFormat(),
+			1,
+			TexCreateFlags,
+			TexCreate_RenderTargetable,
+			false,
+			CreateInfo,
+			RenderTargetTextureRHI,
+			Texture2DRHI
+		);
+
+		TextureRHI = (FTextureRHIRef&)Texture2DRHI;
+		RHIUpdateTextureReference(Owner->TextureReference.TextureReferenceRHI, TextureRHI);
+
+		// add to the list of global deferred updates (updated during scene rendering)
+		// since the latest decoded movie frame is rendered to this media texture target
+		AddToDeferredUpdateList(false);
+	}
+
+	// Create the sampler state RHI resource.
+	FSamplerStateInitializerRHI SamplerStateInitializer(
+		GSystemSettings.TextureLODSettings.GetSamplerFilter(Owner),
+		Owner->AddressX == TA_Wrap ? AM_Wrap : (Owner->AddressX == TA_Clamp ? AM_Clamp : AM_Mirror),
+		Owner->AddressY == TA_Wrap ? AM_Wrap : (Owner->AddressY == TA_Clamp ? AM_Clamp : AM_Mirror),
+		AM_Wrap
+	);
+
+	SamplerStateRHI = RHICreateSamplerState(SamplerStateInitializer);
+}
+
+
+void FMediaTextureResource::ReleaseDynamicRHI( )
+{
+	// release the FTexture RHI resources here as well
+	ReleaseRHI();
+
+	RHIUpdateTextureReference(Owner->TextureReference.TextureReferenceRHI, FTextureRHIParamRef());
+	Texture2DRHI.SafeRelease();
+	RenderTargetTextureRHI.SafeRelease();
+
+	// remove from global list of deferred updates
+	RemoveFromDeferredUpdateList();
+}
+
+
+/* FRenderTarget overrides
+ *****************************************************************************/
+
+FIntPoint FMediaTextureResource::GetSizeXY( ) const
+{
+	return Owner->GetDimensions();
+}
+
+
+/* FDeferredUpdateResource overrides
+ *****************************************************************************/
+
+void FMediaTextureResource::UpdateDeferredResource( FRHICommandListImmediate& RHICmdList )
+{
+	FTimespan CurrentFrameTime = VideoBuffer->GetCurrentSampleTime();
+	TSharedPtr<TArray<uint8>, ESPMode::ThreadSafe> CurrentFrame = VideoBuffer->GetCurrentSample();
+
+	if (CurrentFrame.IsValid())
+	{
+		// draw the latest video frame
+		if (CurrentFrameTime != LastFrameTime)
+		{
+			uint32 Stride = 0;
+			FRHITexture2D* Texture2D = TextureRHI->GetTexture2D();
+			uint8* TextureBuffer = (uint8*)RHILockTexture2D(Texture2D, 0, RLM_WriteOnly, Stride, false);
+
+			FMemory::Memcpy(TextureBuffer, CurrentFrame->GetData(), CurrentFrame->Num());
+			RHIUnlockTexture2D(Texture2D, 0, false);
+
+			LastFrameTime = CurrentFrameTime;
+			Cleared = false;
+		}
+	}
+	else if (!Cleared || (LastClearColor != Owner->ClearColor))
+	{
+		// clear texture if video track selected
+		FRHICommandListImmediate& CommandList = FRHICommandListExecutor::GetImmediateCommandList();
+
+ 		SetRenderTarget(CommandList, RenderTargetTextureRHI, FTextureRHIRef());
+ 		CommandList.SetViewport(0, 0, 0.0f, Owner->GetSurfaceWidth(), Owner->GetSurfaceHeight(), 1.0f);
+ 		CommandList.Clear(true, Owner->ClearColor, false, 0.f, false, 0, FIntRect());
+		CommandList.CopyToResolveTarget(Texture2DRHI, TextureRHI, true, FResolveParams());
+
+		LastClearColor = Owner->ClearColor;
+		Cleared = true;
+	}
+}
diff --git a/Engine/Source/Runtime/MediaAssets/Private/Assets/MediaTextureResource.h b/Engine/Source/Runtime/MediaAssets/Private/Assets/MediaTextureResource.h
new file mode 100644
--- /dev/null
+++ b/Engine/Source/Runtime/MediaAssets/Private/Assets/MediaTextureResource.h
@@ -0,0 +1,62 @@
+// Copyright 1998-2014 Epic Games, Inc. All Rights Reserved.
+
+#pragma once
+
+
+/**
+ * FTextureResource type for movie textures.
+ */
+class FMediaTextureResource
+	: public FTextureResource
+	, public FRenderTarget
+	, public FDeferredUpdateResource
+{
+public:
+
+	/** 
+	 * Creates and initializes a new instance.
+	 *
+	 * @param InOwner The Movie texture object to create a resource for (must not be nullptr).
+	 * @param InVideoBuffer The video sample buffer to use.
+	 */
+	FMediaTextureResource( const class UMediaTexture* InOwner, const FMediaSampleBufferRef& InVideoBuffer );
+
+public:
+
+	// FTextureResource overrides
+
+	virtual void InitDynamicRHI( ) override;
+	virtual void ReleaseDynamicRHI( ) override;
+
+public:
+
+	// FRenderTarget overrides
+
+	virtual FIntPoint GetSizeXY( ) const override;
+
+public:
+
+	// FDeferredUpdateResource overrides
+
+	virtual void UpdateDeferredResource( FRHICommandListImmediate& RHICmdList ) override;
+
+private:
+
+	// Whether the resource was cleared last frame.
+	bool Cleared;
+
+	// The color that the resource was cleared with.
+	FLinearColor LastClearColor;
+
+	// The playback time of the last drawn video frame.
+	FTimespan LastFrameTime;
+
+	/** The UTextureRenderTarget2D which this resource represents. */
+	const class UMediaTexture* Owner;
+
+	/** Texture resource used for rendering with and resolving to. */
+	FTexture2DRHIRef Texture2DRHI;
+
+	/** Pointer to the video sample buffer. */
+	FMediaSampleBufferRef VideoBuffer;
+};
diff --git a/Engine/Source/Runtime/MediaAssets/Private/MediaAssetsModule.cpp b/Engine/Source/Runtime/MediaAssets/Private/MediaAssetsModule.cpp
new file mode 100644
--- /dev/null
+++ b/Engine/Source/Runtime/MediaAssets/Private/MediaAssetsModule.cpp
@@ -0,0 +1,70 @@
+// Copyright 1998-2014 Epic Games, Inc. All Rights Reserved.
+
+#include "MediaAssetsPrivatePCH.h"
+#include "ModuleManager.h"
+
+
+/**
+ * Implements the MediaAssets module.
+ */
+class FMediaAssetsModule
+	: public FSelfRegisteringExec
+	, public IModuleInterface
+{
+public:
+
+	// FSelfRegisteringExec interface
+
+	virtual bool Exec( UWorld* InWorld, const TCHAR* Cmd, FOutputDevice& Ar ) override
+	{
+		if (FParse::Command(&Cmd, TEXT("MOVIE")))
+		{
+			FString MovieCmd = FParse::Token(Cmd, 0);
+
+			if (MovieCmd.Contains(TEXT("PLAY")))
+			{
+				for (TObjectIterator<UMediaAsset> It; It; ++It)
+				{
+					UMediaAsset* MediaAsset = *It;
+					MediaAsset->Play();
+				}
+			}
+			else if (MovieCmd.Contains(TEXT("PAUSE")))
+			{
+				for (TObjectIterator<UMediaAsset> It; It; ++It)
+				{
+					UMediaAsset* MediaAsset = *It;
+					MediaAsset->Pause();
+				}
+			}
+			else if (MovieCmd.Contains(TEXT("STOP")))
+			{
+				for (TObjectIterator<UMediaAsset> It; It; ++It)
+				{
+					UMediaAsset* MediaAsset = *It;
+					MediaAsset->Stop();
+				}
+			}
+
+			return true;
+		}
+
+		return false;
+	}
+
+public:
+
+	// IModuleInterface interface
+
+	virtual void StartupModule( ) override { }
+
+	virtual void ShutdownModule( ) override { }
+
+	virtual bool SupportsDynamicReloading( ) override
+	{
+		return false;
+	}
+};
+
+
+IMPLEMENT_MODULE(FMediaAssetsModule, MediaAssets);
diff --git a/Engine/Source/Runtime/MediaAssets/Private/MediaAssetsPrivatePCH.h b/Engine/Source/Runtime/MediaAssets/Private/MediaAssetsPrivatePCH.h
new file mode 100644
--- /dev/null
+++ b/Engine/Source/Runtime/MediaAssets/Private/MediaAssetsPrivatePCH.h
@@ -0,0 +1,17 @@
+// Copyright 1998-2014 Epic Games, Inc. All Rights Reserved.
+
+#pragma once
+
+
+/* Private dependencies
+ *****************************************************************************/
+
+#include "Core.h"
+#include "CoreUObject.h"
+#include "Engine.h"
+#include "Media.h"
+
+#include "MaterialExpressionTExtureSampleParameterMedia.h"
+#include "MediaAsset.h"
+#include "MediaSoundWave.h"
+#include "MediaTexture.h"
diff --git a/Engine/Source/Runtime/MediaAssets/Public/MaterialExpressionTextureSampleParameterMedia.h b/Engine/Source/Runtime/MediaAssets/Public/MaterialExpressionTextureSampleParameterMedia.h
new file mode 100644
--- /dev/null
+++ b/Engine/Source/Runtime/MediaAssets/Public/MaterialExpressionTextureSampleParameterMedia.h
@@ -0,0 +1,29 @@
+// Copyright 1998-2014 Epic Games, Inc. All Rights Reserved.
+
+#pragma once
+
+#include "Materials/MaterialExpressionTextureSampleParameter.h"
+#include "MaterialExpressionTextureSampleParameterMedia.generated.h"
+
+
+/**
+ * Implements a material instance parameter for media textures.
+ */
+UCLASS(collapsecategories, hidecategories=Object)
+class MEDIAASSETS_API UMaterialExpressionTextureSampleParameterMedia
+	: public UMaterialExpressionTextureSampleParameter
+{
+	GENERATED_UCLASS_BODY()
+
+public:
+
+	// UMaterialExpressionTextureSampleParameter interface
+	virtual const TCHAR* GetRequirements( ) override;
+	virtual bool TextureIsValid( UTexture* InTexture ) override;
+
+public:
+
+	// UMaterialExpression Interface
+
+	virtual void GetCaption( TArray<FString>& OutCaptions ) const override;
+};
diff --git a/Engine/Source/Runtime/MediaAssets/Public/MediaAsset.h b/Engine/Source/Runtime/MediaAssets/Public/MediaAsset.h
new file mode 100644
--- /dev/null
+++ b/Engine/Source/Runtime/MediaAssets/Public/MediaAsset.h
@@ -0,0 +1,294 @@
+// Copyright 1998-2014 Epic Games, Inc. All Rights Reserved.
+
+#pragma once
+
+#include "MediaAsset.generated.h"
+
+
+/**
+ * Enumerates available media streaming modes.
+ */
+UENUM()
+enum EMediaAssetStreamModes
+{
+	/** Load media contents to memory before playing. */
+	MASM_FromMemory UMETA(DisplayName="FromMemory"),
+
+	/** Play media directly from the URL. */
+	MASM_FromUrl UMETA(DisplayName="FromUrl"),
+
+	MASM_MAX,
+};
+
+
+/** Multicast delegate that is invoked when a media asset's media is being closed. */
+DECLARE_DYNAMIC_MULTICAST_DELEGATE_OneParam(FOnMediaAssetMediaClosing, FString, ClosedUrl);
+
+/** Multicast delegate that is invoked when a media asset's media has been opened. */
+DECLARE_DYNAMIC_MULTICAST_DELEGATE_OneParam(FOnMediaAssetMediaOpened, FString, OpenedUrl);
+
+
+/**
+ * Implements a media asset.
+ *
+ * This class is represents a media URL along with a corresponding media player
+ * for exposing media playback functionality to the Engine and to Blueprints.
+ */
+UCLASS(hidecategories=(Object))
+class MEDIAASSETS_API UMediaAsset
+	: public UObject
+{
+	GENERATED_UCLASS_BODY()
+
+	/** Whether this media should automatically start playing when it is loaded. */
+	UPROPERTY(EditAnywhere, Category=Playback)
+	uint32 AutoPlay:1;
+
+	/** Whether playback should loop when it reaches the end. */
+	UPROPERTY(EditAnywhere, BlueprintReadWrite, Category=Playback)
+	uint32 Looping:1;
+
+	/** The media's playback rate (1.0 = real time) */
+	UPROPERTY(EditAnywhere, BlueprintReadOnly, Category=Playback)
+	float PlaybackRate;
+
+	/** The current playback position. */
+	UPROPERTY(BlueprintReadOnly, Category=Playback)
+	FTimespan Position;
+
+	/** Whether the media should reset to first frame when playback reaches the end. */
+	UPROPERTY(EditAnywhere, BlueprintReadWrite, Category=Playback)
+	uint32 ResetOnLastFrame:1;
+
+	/** Select where to stream the media from, i.e. file or memory. */
+	UPROPERTY(EditAnywhere, BlueprintReadOnly, Category=MediaSource)
+	TEnumAsByte<enum EMediaAssetStreamModes> StreamMode;
+
+	/** The URL to the media file to be played. */
+	UPROPERTY(EditAnywhere, BlueprintReadOnly, Category=MediaSource)
+	FFilePath URL;
+
+public:
+
+	/**
+	 * Checks whether media playback can be paused right now.
+	 *
+	 * Playback can be paused if the media supports pausing and if it is currently playing.
+	 *
+	 * @return true if pausing playback can be paused, false otherwise.
+	 * @see CanPlay, CanStop, Pause
+	 */
+	UFUNCTION(BlueprintCallable, Category="Media|MediaAsset")
+	bool CanPause( ) const;
+
+	/**
+	 * Checks whether media playback can be started right now.
+	 *
+	 * @return true if playback can be started, false otherwise.
+	 * @see CanPause, CanStop, Play
+	 */
+	UFUNCTION(BlueprintCallable, Category="Media|MediaAsset")
+	bool CanPlay( ) const;
+
+	/**
+	 * Checks whether playback can be stopped.
+	 *
+	 * @return true if stopping is supported, false otherwise.
+	 * @see CanPause, CanPlay, Stop
+	 */
+	UFUNCTION(BlueprintCallable, Category="Media|MediaAsset")
+	bool CanStop( ) const;
+
+	/**
+	 * Gets the media's duration.
+	 *
+	 * @return A time span representing the duration.
+	 * @see GetPosition
+	 */
+	UFUNCTION(BlueprintCallable, Category="Media|MediaAsset")
+	FTimespan GetDuration( ) const;
+
+	/**
+	 * Gets the media's current playback rate.
+	 *
+	 * @return The playback rate.
+	 */
+	UFUNCTION(BlueprintCallable, Category="Media|MediaAsset")
+	float GetRate( ) const;
+
+	/**
+	 * Gets the media's current playback time.
+	 *
+	 * @return Playback time.
+	 * @see GetDuration, Seek
+	 */
+	UFUNCTION(BlueprintCallable, Category="Media|MediaAsset")
+	FTimespan GetTime( ) const;
+
+	/**
+	 * Checks whether playback is currently paused.
+	 *
+	 * @return true if playback is paused, false otherwise.
+	 * @see CanPause, IsPlaying, IsStopped, Pause
+	 */
+	UFUNCTION(BlueprintCallable, Category="Media|MediaAsset")
+	bool IsPaused( ) const;
+
+	/**
+	 * Checks whether playback has started.
+	 *
+	 * @return true if playback has started, false otherwise.
+	 * @see CanPlay, IsPaused, IsStopped, Play
+	 */
+	UFUNCTION(BlueprintCallable, Category="Media|MediaAsset")
+	bool IsPlaying( ) const;
+
+	/**
+	 * Checks whether playback has stopped.
+	 *
+	 * @return true if playback has stopped, false otherwise.
+	 * @see CanStop, IsPaused, IsPlaying, Stop
+	 */
+	UFUNCTION(BlueprintCallable, Category="Media|MediaAsset")
+	bool IsStopped( ) const;
+
+	/**
+	 * Pauses media playback.
+	 *
+	 * @return true if playback is being paused, false otherwise.
+	 * @see CanPause, Play, Rewind, Seek, Stop
+	 */
+	UFUNCTION(BlueprintCallable, Category="Media|MediaAsset")
+	bool Pause( );
+	
+	/**
+	 * Starts media playback.
+	 *
+	 * @return true if rewinding, false otherwise.
+	 * @see GetTime, Pause, Play, Seek, Stop
+	 */
+	UFUNCTION(BlueprintCallable, Category="Media|MediaAsset")
+	bool Rewind( );
+
+	/**
+	 * Rewinds the media.
+	 *
+	 * @return true if playback is starting, false otherwise.
+	 * @see CanPlay, Pause, Rewind, Seek, Stop
+	 */
+	UFUNCTION(BlueprintCallable, Category="Media|MediaAsset")
+	bool Play( );
+
+	/**
+	 * Seeks to the specified playback time.
+	 *
+	 * @param InTime The playback time to set.
+	 * @return true on success, false otherwise.
+	 * @see GetTime
+	 */
+	UFUNCTION(BlueprintCallable, Category="Media|MediaAsset")
+	bool Seek( const FTimespan& InTime );
+
+	/**
+	 * Changes the media's playback rate.
+	 *
+	 * @param Rate The playback rate to set.
+	 * @return true on success, false otherwise.
+	 * @see GetRate
+	 */
+	UFUNCTION(BlueprintCallable, Category="Media|MediaAsset")
+	bool SetRate( float InRate );
+
+	/**
+	 * Stops movie playback.
+	 *
+	 * @see CanStop, IsStopped, Pause, Play, Seek
+	 */
+	UFUNCTION(BlueprintCallable, Category="Media|MediaAsset")
+	void Stop( );
+
+	/**
+	 * Checks whether the specified playback rate is supported.
+	 *
+	 * @param Rate The playback rate to check.
+	 * @param Unthinned Whether no frames should be dropped at the given rate.
+	 * @see SupportsScrubbing, SupportsSeeking
+	 */
+	UFUNCTION(BlueprintCallable, Category="Media|MediaAsset")
+	bool SupportsRate( float Rate, bool Unthinned ) const;
+
+	/**
+	 * Checks whether the currently loaded media supports scrubbing.
+	 *
+	 * @return true if scrubbing is supported, false otherwise.
+	 * @see SupportsRate, SupportsSeeking
+	 */
+	UFUNCTION(BlueprintCallable, Category="Media|MediaAsset")
+	bool SupportsScrubbing( ) const;
+
+	/**
+	 * Checks whether the currently loaded media can jump to a certain position.
+	 *
+	 * @return true if seeking is supported, false otherwise.
+	 * @see SupportsRate, SupportsScrubbing
+	 */
+	UFUNCTION(BlueprintCallable, Category="Media|MediaAsset")
+	bool SupportsSeeking( ) const;
+
+public:
+
+	/** Holds a delegate that is invoked when a media asset's media is being unloaded. */
+	UPROPERTY(BlueprintAssignable, Category="Media|MediaAsset")
+	FOnMediaAssetMediaClosing OnMediaClosing;
+
+	/** Holds a delegate that is invoked when a media asset's media has been opened. */
+	UPROPERTY(BlueprintAssignable, Category="Media|MediaAsset")
+	FOnMediaAssetMediaOpened OnMediaOpened;
+
+public:
+
+	/**
+	 * Gets the asset's media player object.
+	 *
+	 * @return Media player, or nullptr if no player was created.
+	 */
+	TSharedPtr<class IMediaPlayer> GetMediaPlayer( ) const
+	{
+		return MediaPlayer;
+	}
+
+public:
+
+	// UObject overrides.
+
+	virtual void BeginDestroy( ) override;
+	virtual FString GetDesc( ) override;
+	virtual void PostLoad( ) override;
+
+#if WITH_EDITOR
+	virtual void PostEditChangeProperty( FPropertyChangedEvent& PropertyChangedEvent ) override;
+#endif
+
+protected:
+
+	/**
+	 * Initializes the media player.
+	 */
+	void InitializePlayer( );
+
+private:
+
+	// Callback for when the media player is unloading media.
+	void HandleMediaPlayerMediaClosing( FString ClosedUrl );
+
+	// Callback for when the media player is loading media.
+	void HandleMediaPlayerMediaOpened( FString OpenedUrl );
+
+private:
+
+	/** Holds the URL of the currently loaded media file. */
+	FString CurrentUrl;
+
+	/** Holds the media player used to play the media file. */
+	TSharedPtr<class IMediaPlayer> MediaPlayer;
+};
diff --git a/Engine/Source/Runtime/MediaAssets/Public/MediaSoundWave.h b/Engine/Source/Runtime/MediaAssets/Public/MediaSoundWave.h
new file mode 100644
--- /dev/null
+++ b/Engine/Source/Runtime/MediaAssets/Public/MediaSoundWave.h
@@ -0,0 +1,97 @@
+// Copyright 1998-2014 Epic Games, Inc. All Rights Reserved.
+
+#pragma once
+
+#include "MediaSampleQueue.h"
+#include "Sound/SoundWave.h"
+#include "MediaSoundWave.generated.h"
+
+
+/**
+ * Implements a playable sound object for audio streams from media assets.
+ */
+UCLASS(hidecategories=(Compression, Sound, SoundWave, Subtitles))
+class MEDIAASSETS_API UMediaSoundWave
+	: public USoundWave
+{
+	GENERATED_UCLASS_BODY()
+
+	/** The index of the media asset's audio track to get the wave data from. */
+	UPROPERTY(EditAnywhere, BlueprintReadWrite, Category=MediaAsset)
+	int32 AudioTrackIndex;
+
+	/** The media asset to stream audio from. */
+	UPROPERTY(EditAnywhere, BlueprintReadWrite, Category=MediaAsset)
+	UMediaAsset* MediaAsset;
+
+public:
+
+	/** Destructor. */
+	~UMediaSoundWave( );
+
+public:
+
+	/**
+	 * Sets the media asset to be used for this texture.
+	 *
+	 * @param InMediaAsset The asset to set.
+	 */
+	UFUNCTION(BlueprintCallable, Category="Media|MediaSound")
+	void SetMediaAsset( UMediaAsset* InMediaAsset );
+
+public:
+
+	/**
+	 * Gets the media player associated with the assigned media asset.
+	 *
+	 * @return Media player, or nullptr if no player is available.
+	 */
+	TSharedPtr<class IMediaPlayer> GetMediaPlayer( ) const;
+
+public:
+
+	// USoundWave overrides
+
+	virtual int32 GeneratePCMData(uint8* PCMData, const int32 SamplesNeeded) override;
+	virtual FByteBulkData* GetCompressedData(FName Format) override;
+	virtual int32 GetResourceSizeForFormat( FName Format ) override;
+	virtual void InitAudioResource( FByteBulkData& CompressedData ) override;
+	virtual bool InitAudioResource( FName Format ) override;
+
+public:
+
+	// UObject overrides
+
+	virtual void GetAssetRegistryTags( TArray<FAssetRegistryTag>& OutTags ) const override;
+	virtual SIZE_T GetResourceSize( EResourceSizeMode::Type Mode ) override;
+	virtual void Serialize( FArchive& Ar ) override;
+
+protected:
+
+	/**
+	 * Initializes the audio stream.
+	 */
+	void InitializeStream( );
+
+private:
+
+	// Callback for when the media player unloads media.
+	void HandleMediaPlayerClosing( FString ClosingUrl );
+
+	// Callback for when the media player unloads media.
+	void HandleMediaPlayerOpened( FString OpenedUrl );
+
+private:
+
+	/** The audio sample queue. */
+	FMediaSampleQueueRef AudioQueue;
+
+	// Holds the selected audio track.
+	IMediaTrackPtr AudioTrack;
+
+	// Holds the media asset currently being used.
+	UMediaAsset* CurrentMediaAsset;
+
+	// Holds queued audio samples
+	TArray<uint8> QueuedAudio;
+};
diff --git a/Engine/Source/Runtime/MediaAssets/Public/MediaTexture.h b/Engine/Source/Runtime/MediaAssets/Public/MediaTexture.h
new file mode 100644
--- /dev/null
+++ b/Engine/Source/Runtime/MediaAssets/Public/MediaTexture.h
@@ -0,0 +1,153 @@
+// Copyright 1998-2014 Epic Games, Inc. All Rights Reserved.
+
+#pragma once
+
+#include "MediaSampleBuffer.h"
+#include "MediaTexture.generated.h"
+
+
+/**
+ * Implements a texture for rendering video tracks from media assets.
+ */
+UCLASS(hidecategories=(Compression, LevelOfDetail, Object, Texture))
+class MEDIAASSETS_API UMediaTexture
+	: public UTexture
+{
+	GENERATED_UCLASS_BODY()
+
+	/** The addressing mode to use for the X axis. */
+	UPROPERTY(EditAnywhere, BlueprintReadWrite, Category=MediaTexture, AssetRegistrySearchable)
+	TEnumAsByte<enum TextureAddress> AddressX;
+
+	/** The addressing mode to use for the Y axis. */
+	UPROPERTY(EditAnywhere, BlueprintReadWrite, Category=MediaTexture, AssetRegistrySearchable)
+	TEnumAsByte<enum TextureAddress> AddressY;
+
+	/** The color used to clear the texture if no video data is drawn. */
+	UPROPERTY(EditAnywhere, BlueprintReadWrite, Category=MediaTexture)
+	FLinearColor ClearColor;
+
+	/** The index of the media asset's video track to render on this texture. */
+	UPROPERTY(EditAnywhere, BlueprintReadWrite, Category=MediaAsset)
+	int32 VideoTrackIndex;
+
+public:
+
+	/** Destructor. */
+	~UMediaTexture( );
+
+public:
+
+	/**
+	 * Sets the media asset to be used for this texture.
+	 *
+	 * @param InMediaAsset The asset to set.
+	 */
+	UFUNCTION(BlueprintCallable, Category="Media|MediaTexture")
+	void SetMediaAsset( UMediaAsset* InMediaAsset );
+
+public:
+
+	/**
+	 * Gets the width and height of the texture (in pixels).
+	 *
+	 * @return Texture dimensions.
+	 */
+	FIntPoint GetDimensions( ) const
+	{
+		return CachedDimensions;
+	}
+
+	/**
+	 * Gets the texture's pixel format.
+	 *
+	 * @return Pixel format (always PF_B8G8R8A8 for all movie textures).
+	 */
+	TEnumAsByte<enum EPixelFormat> GetFormat( ) const
+	{
+		return PF_B8G8R8A8;
+	}
+
+	/**
+	 * Gets the media player associated with the assigned media asset.
+	 *
+	 * @return Media player, or nullptr if no player is available.
+	 */
+	TSharedPtr<class IMediaPlayer> GetMediaPlayer( ) const
+	{
+		return (MediaAsset != nullptr) ? MediaAsset->GetMediaPlayer() : nullptr;
+	}
+
+	/**
+	 * Gets the currently selected video track, if any.
+	 *
+	 * @return The selected video track, or nullptr if none is selected.
+	 */
+	TSharedPtr<class IMediaTrack, ESPMode::ThreadSafe> GetVideoTrack( ) const
+	{
+		return VideoTrack;
+	}
+	
+public:
+
+	// UTexture overrides.
+
+	virtual FTextureResource* CreateResource( ) override;
+	virtual EMaterialValueType GetMaterialType( ) override;
+	virtual float GetSurfaceWidth( ) const override;
+	virtual float GetSurfaceHeight( ) const override;
+
+public:
+
+	// UObject overrides.
+
+	virtual void BeginDestroy( ) override;
+	virtual void FinishDestroy( ) override;
+	virtual FString GetDesc( ) override;
+	virtual SIZE_T GetResourceSize( EResourceSizeMode::Type Mode ) override;
+	virtual bool IsReadyForFinishDestroy( ) override;
+	virtual void PostLoad( ) override;
+
+#if WITH_EDITOR
+	virtual void PreEditChange( UProperty* PropertyAboutToChange ) override;
+	virtual void PostEditChangeProperty( FPropertyChangedEvent& PropertyChangedEvent ) override;
+#endif
+
+protected:
+
+	/** The media asset to track video from. */
+	UPROPERTY(EditAnywhere, BlueprintReadOnly, Category=MediaAsset)
+	UMediaAsset* MediaAsset;
+
+protected:
+
+	/**
+	 * Initializes the video track.
+	 */
+	void InitializeTrack( );
+
+private:
+
+	// Callback for when the media player unloads media.
+	void HandleMediaPlayerClosing( FString ClosingUrl );
+
+	// Callback for when the media player unloads media.
+	void HandleMediaPlayerOpened( FString OpenedUrl );
+
+private:
+
+	// The texture's cached width and height (in pixels).
+	FIntPoint CachedDimensions;
+
+	// Holds the media asset currently being used.
+	UMediaAsset* CurrentMediaAsset;
+
+	// Synchronizes access to this object from the render thread.
+	FRenderCommandFence* ReleasePlayerFence;
+
+	/** The video sample buffer. */
+	FMediaSampleBufferRef VideoBuffer;
+
+	// Holds the selected video track.
+	IMediaTrackPtr VideoTrack;
+};
