Media: added missing virtual destructor; cosmetic changes

diff --git a/Engine/Source/Runtime/Media/Public/Common/MediaSampleBuffer.h b/Engine/Source/Runtime/Media/Public/Common/MediaSampleBuffer.h
--- a/Engine/Source/Runtime/Media/Public/Common/MediaSampleBuffer.h
+++ b/Engine/Source/Runtime/Media/Public/Common/MediaSampleBuffer.h
@@ -64,9 +64,9 @@ public:
 	
 private:
 
-	// Holds the sample data of the current media sample.
+	/** Holds the sample data of the current media sample. */
 	TSharedPtr<TArray<uint8>, ESPMode::ThreadSafe> CurrentSample;
 
-	// The playback time of the currently available media sample.
+	/** The playback time of the currently available media sample. */
 	FTimespan CurrentSampleTime;
 };
diff --git a/Engine/Source/Runtime/Media/Public/Interfaces/IMediaModule.h b/Engine/Source/Runtime/Media/Public/Interfaces/IMediaModule.h
--- a/Engine/Source/Runtime/Media/Public/Interfaces/IMediaModule.h
+++ b/Engine/Source/Runtime/Media/Public/Interfaces/IMediaModule.h
@@ -67,7 +67,6 @@ public:
 	 * Gets a reference to the messaging module instance.
 	 *
 	 * @todo gmp: better implementation using dependency injection.
-	 *
 	 * @return A reference to the Media module.
 	 */
 	static IMediaModule& Get( )
@@ -82,8 +81,6 @@ public:
 
 public:
 
-	/**
-	 * Virtual destructor.
-	 */
+	/** Virtual destructor. */
 	virtual ~IMediaModule( ) { }
 };
diff --git a/Engine/Source/Runtime/Media/Public/Interfaces/IMediaPlayer.h b/Engine/Source/Runtime/Media/Public/Interfaces/IMediaPlayer.h
--- a/Engine/Source/Runtime/Media/Public/Interfaces/IMediaPlayer.h
+++ b/Engine/Source/Runtime/Media/Public/Interfaces/IMediaPlayer.h
@@ -273,8 +273,6 @@ public:
 
 public:
 
-	/**
-	 * Virtual destructor.
-	 */
+	/** Virtual destructor. */
 	virtual ~IMediaPlayer( ) { }
 };
diff --git a/Engine/Source/Runtime/Media/Public/Interfaces/IMediaPlayerFactory.h b/Engine/Source/Runtime/Media/Public/Interfaces/IMediaPlayerFactory.h
--- a/Engine/Source/Runtime/Media/Public/Interfaces/IMediaPlayerFactory.h
+++ b/Engine/Source/Runtime/Media/Public/Interfaces/IMediaPlayerFactory.h
@@ -55,8 +55,6 @@ public:
 
 public:
 
-	/**
-	 * Virtual destructor.
-	 */
+	/** Virtual destructor. */
 	virtual ~IMediaPlayerFactory( ) { }
 };
diff --git a/Engine/Source/Runtime/Media/Public/Interfaces/IMediaSink.h b/Engine/Source/Runtime/Media/Public/Interfaces/IMediaSink.h
--- a/Engine/Source/Runtime/Media/Public/Interfaces/IMediaSink.h
+++ b/Engine/Source/Runtime/Media/Public/Interfaces/IMediaSink.h
@@ -34,8 +34,6 @@ public:
 
 public:
 
-	/**
-	 * Virtual destructor.
-	 */
+	/** Virtual destructor. */
 	~IMediaSink( ) { }
 };
diff --git a/Engine/Source/Runtime/Media/Public/Interfaces/IMediaTrack.h b/Engine/Source/Runtime/Media/Public/Interfaces/IMediaTrack.h
--- a/Engine/Source/Runtime/Media/Public/Interfaces/IMediaTrack.h
+++ b/Engine/Source/Runtime/Media/Public/Interfaces/IMediaTrack.h
@@ -176,4 +176,9 @@ public:
 	 * @param Sink The sink to remove.
 	 */
 	virtual void RemoveSink( const IMediaSinkRef& Sink ) = 0;
+
+public:
+
+	/** Virtual destructor. */
+	virtual ~IMediaTrack( ) { }
 };
diff --git a/Engine/Source/Runtime/Media/Public/Interfaces/IMediaTrackAudioDetails.h b/Engine/Source/Runtime/Media/Public/Interfaces/IMediaTrackAudioDetails.h
--- a/Engine/Source/Runtime/Media/Public/Interfaces/IMediaTrackAudioDetails.h
+++ b/Engine/Source/Runtime/Media/Public/Interfaces/IMediaTrackAudioDetails.h
@@ -26,8 +26,6 @@ public:
 
 public:
 
-	/**
-	 * Virtual destructor.
-	 */
+	/** Virtual destructor. */
 	virtual ~IMediaTrackAudioDetails( ) { }
 };
diff --git a/Engine/Source/Runtime/Media/Public/Interfaces/IMediaTrackCaptionDetails.h b/Engine/Source/Runtime/Media/Public/Interfaces/IMediaTrackCaptionDetails.h
--- a/Engine/Source/Runtime/Media/Public/Interfaces/IMediaTrackCaptionDetails.h
+++ b/Engine/Source/Runtime/Media/Public/Interfaces/IMediaTrackCaptionDetails.h
@@ -10,8 +10,6 @@ class IMediaTrackCaptionDetails
 {
 public:
 
-	/**
-	 * Virtual destructor.
-	 */
+	/** Virtual destructor. */
 	virtual ~IMediaTrackCaptionDetails( ) { }
 };
diff --git a/Engine/Source/Runtime/Media/Public/Interfaces/IMediaTrackVideoDetails.h b/Engine/Source/Runtime/Media/Public/Interfaces/IMediaTrackVideoDetails.h
--- a/Engine/Source/Runtime/Media/Public/Interfaces/IMediaTrackVideoDetails.h
+++ b/Engine/Source/Runtime/Media/Public/Interfaces/IMediaTrackVideoDetails.h
@@ -54,8 +54,6 @@ public:
 
 public:
 
-	/**
-	 * Virtual destructor.
-	 */
+	/** Virtual destructor. */
 	virtual ~IMediaTrackVideoDetails( ) { }
 };
diff --git a/Engine/Source/Runtime/MediaAssets/Public/MediaAsset.h b/Engine/Source/Runtime/MediaAssets/Public/MediaAsset.h
--- a/Engine/Source/Runtime/MediaAssets/Public/MediaAsset.h
+++ b/Engine/Source/Runtime/MediaAssets/Public/MediaAsset.h
@@ -271,9 +271,7 @@ public:
 
 protected:
 
-	/**
-	 * Initializes the media player.
-	 */
+	/** Initializes the media player. */
 	void InitializePlayer( );
 
 private:
diff --git a/Engine/Source/Runtime/MediaAssets/Public/MediaSoundWave.h b/Engine/Source/Runtime/MediaAssets/Public/MediaSoundWave.h
--- a/Engine/Source/Runtime/MediaAssets/Public/MediaSoundWave.h
+++ b/Engine/Source/Runtime/MediaAssets/Public/MediaSoundWave.h
@@ -68,9 +68,7 @@ public:
 
 protected:
 
-	/**
-	 * Initializes the audio stream.
-	 */
+	/** Initializes the audio stream. */
 	void InitializeStream( );
 
 private:
@@ -86,12 +84,12 @@ private:
 	/** The audio sample queue. */
 	FMediaSampleQueueRef AudioQueue;
 
-	// Holds the selected audio track.
+	/** Holds the selected audio track. */
 	IMediaTrackPtr AudioTrack;
 
-	// Holds the media asset currently being used.
+	/** Holds the media asset currently being used. */
 	UMediaAsset* CurrentMediaAsset;
 
-	// Holds queued audio samples
+	/** Holds queued audio samples. */
 	TArray<uint8> QueuedAudio;
 };
diff --git a/Engine/Source/Runtime/MediaAssets/Public/MediaTexture.h b/Engine/Source/Runtime/MediaAssets/Public/MediaTexture.h
--- a/Engine/Source/Runtime/MediaAssets/Public/MediaTexture.h
+++ b/Engine/Source/Runtime/MediaAssets/Public/MediaTexture.h
@@ -136,18 +136,18 @@ private:
 
 private:
 
-	// The texture's cached width and height (in pixels).
+	/** The texture's cached width and height (in pixels). */
 	FIntPoint CachedDimensions;
 
-	// Holds the media asset currently being used.
+	/** Holds the media asset currently being used. */
 	UMediaAsset* CurrentMediaAsset;
 
-	// Synchronizes access to this object from the render thread.
+	/** Synchronizes access to this object from the render thread. */
 	FRenderCommandFence* ReleasePlayerFence;
 
 	/** The video sample buffer. */
 	FMediaSampleBufferRef VideoBuffer;
 
-	// Holds the selected video track.
+	/** Holds the selected video track. */
 	IMediaTrackPtr VideoTrack;
 };
