Media: Playing media assets from URL instead of memory by default

diff --git a/Engine/Source/Runtime/MediaAssets/Private/Assets/MediaAsset.cpp b/Engine/Source/Runtime/MediaAssets/Private/Assets/MediaAsset.cpp
--- a/Engine/Source/Runtime/MediaAssets/Private/Assets/MediaAsset.cpp
+++ b/Engine/Source/Runtime/MediaAssets/Private/Assets/MediaAsset.cpp
@@ -12,7 +12,7 @@ UMediaAsset::UMediaAsset( const class FPostConstructInitializeProperties& PCIP )
 	, Looping(true)
 	, PlaybackRate(1.0f)
 	, ResetOnLastFrame(false)
-	, StreamMode(MASM_FromMemory)
+	, StreamMode(MASM_FromUrl)
 	, MediaPlayer(nullptr)
 { }
 
