Media: Implemented basic playback looping support

diff --git a/Engine/Plugins/Media/MediaAssetEditor/Source/MediaAssetEditor/Private/Factories/MediaAssetFactory.cpp b/Engine/Plugins/Media/MediaAssetEditor/Source/MediaAssetEditor/Private/Factories/MediaAssetFactory.cpp
--- a/Engine/Plugins/Media/MediaAssetEditor/Source/MediaAssetEditor/Private/Factories/MediaAssetFactory.cpp
+++ b/Engine/Plugins/Media/MediaAssetEditor/Source/MediaAssetEditor/Private/Factories/MediaAssetFactory.cpp
@@ -33,8 +33,7 @@ UMediaAssetFactory::UMediaAssetFactory( const class FPostConstructInitializeProp
 UObject* UMediaAssetFactory::FactoryCreateBinary( UClass* Class, UObject* InParent, FName Name, EObjectFlags Flags, UObject* Context, const TCHAR* Type, const uint8*& Buffer, const uint8* BufferEnd, FFeedbackContext* Warn )
 {
 	UMediaAsset* MediaAsset = CastChecked<UMediaAsset>(StaticConstructObject(Class, InParent, Name, Flags));
-	MediaAsset->URL.FilePath = CurrentFilename;
-	MediaAsset->PostEditChange();
+	MediaAsset->OpenUrl(CurrentFilename);
 
 	return MediaAsset;
 }
diff --git a/Engine/Plugins/Media/WmfMedia/Source/WmfMedia/Private/Player/WmfMediaPlayer.cpp b/Engine/Plugins/Media/WmfMedia/Source/WmfMedia/Private/Player/WmfMediaPlayer.cpp
--- a/Engine/Plugins/Media/WmfMedia/Source/WmfMedia/Private/Player/WmfMediaPlayer.cpp
+++ b/Engine/Plugins/Media/WmfMedia/Source/WmfMedia/Private/Player/WmfMediaPlayer.cpp
@@ -76,6 +76,12 @@ FString FWmfMediaPlayer::GetUrl( ) const
 }
 
 
+bool FWmfMediaPlayer::IsLooping( ) const 
+{
+	return ((MediaSession != NULL) && MediaSession->IsLooping());
+}
+
+
 bool FWmfMediaPlayer::IsPaused( ) const
 {
 	if (MediaSession == NULL)
@@ -192,6 +198,19 @@ bool FWmfMediaPlayer::Seek( const FTimespan& Time )
 }
 
 
+bool FWmfMediaPlayer::SetLooping( bool Looping )
+{
+	if (MediaSession == NULL)
+	{
+		return false;
+	}
+
+	MediaSession->SetLooping(Looping);
+
+	return true;
+}
+
+
 bool FWmfMediaPlayer::SupportsRate( float Rate, bool Unthinned ) const
 {
 	return (MediaSession != NULL)
diff --git a/Engine/Plugins/Media/WmfMedia/Source/WmfMedia/Private/Player/WmfMediaPlayer.h b/Engine/Plugins/Media/WmfMedia/Source/WmfMedia/Private/Player/WmfMediaPlayer.h
--- a/Engine/Plugins/Media/WmfMedia/Source/WmfMedia/Private/Player/WmfMediaPlayer.h
+++ b/Engine/Plugins/Media/WmfMedia/Source/WmfMedia/Private/Player/WmfMediaPlayer.h
@@ -33,6 +33,7 @@ public:
 	virtual FTimespan GetTime( ) const override;
 	virtual const TArray<IMediaTrackRef>& GetTracks( ) const override;
 	virtual FString GetUrl( ) const override;
+	virtual bool IsLooping( ) const override;
 	virtual bool IsPaused( ) const override;
 	virtual bool IsPlaying( ) const override;
 	virtual bool IsReady( ) const override;
@@ -40,6 +41,7 @@ public:
 	virtual bool Open( const TSharedRef<TArray<uint8>>& Buffer, const FString& OriginalUrl ) override;
 	virtual bool Play( float Rate ) override;
 	virtual bool Seek( const FTimespan& Time ) override;
+	virtual bool SetLooping( bool Looping ) override;
 	virtual bool SupportsRate( float Rate, bool Unthinned ) const override;
 	virtual bool SupportsScrubbing( ) const override;
 	virtual bool SupportsSeeking( ) const override;
diff --git a/Engine/Plugins/Media/WmfMedia/Source/WmfMedia/Private/Wmf/WmfMediaSession.cpp b/Engine/Plugins/Media/WmfMedia/Source/WmfMedia/Private/Wmf/WmfMediaSession.cpp
--- a/Engine/Plugins/Media/WmfMedia/Source/WmfMedia/Private/Wmf/WmfMediaSession.cpp
+++ b/Engine/Plugins/Media/WmfMedia/Source/WmfMedia/Private/Wmf/WmfMediaSession.cpp
@@ -12,6 +12,7 @@ FWmfMediaSession::FWmfMediaSession( const TComPtr<IMFTopology>& InTopology )
 	, ChangeRequested(false)
 	, CurrentRate(0.0f)
 	, LastError(S_OK)
+	, Looping(false)
 	, RefCount(1)
 	, RequestedPosition(FTimespan::MinValue())
 	, RequestedRate(0.0f)
@@ -179,7 +180,15 @@ STDMETHODIMP FWmfMediaSession::Invoke( IMFAsyncResult* AsyncResult )
 			}
 			else
 			{
-				if (EventType == MESessionCapabilitiesChanged)
+				if (EventType == MEEndOfPresentation)
+				{
+					if (!Looping)
+					{
+						RequestedState = EMediaStates::Stopped;
+						MediaSession->Stop();
+					}
+				}
+				else if (EventType == MESessionCapabilitiesChanged)
 				{
 					Capabilities = ::MFGetAttributeUINT32(Event, MF_EVENT_SESSIONCAPS, Capabilities);
 				}
diff --git a/Engine/Plugins/Media/WmfMedia/Source/WmfMedia/Private/Wmf/WmfMediaSession.h b/Engine/Plugins/Media/WmfMedia/Source/WmfMedia/Private/Wmf/WmfMediaSession.h
--- a/Engine/Plugins/Media/WmfMedia/Source/WmfMedia/Private/Wmf/WmfMediaSession.h
+++ b/Engine/Plugins/Media/WmfMedia/Source/WmfMedia/Private/Wmf/WmfMediaSession.h
@@ -102,6 +102,16 @@ public:
 		return CurrentState;
 	}
 
+	/**
+	 * Checks whether playback is currently looping.
+	 *
+	 * @return true if looping, false otherwise.
+	 */
+	bool IsLooping( ) const
+	{
+		return Looping;
+	}
+
 	/**
 	 * Checks whether the specified playback rate is supported.
 	 *
@@ -111,6 +121,16 @@ public:
 	 */
 	bool IsRateSupported( float Rate, bool Unthinned ) const;
 
+	/**
+	 * Sets whether playback should loop back to the beginning.
+	 *
+	 * @param InLooping Whether playback should be looped.
+	 */
+	void SetLooping( bool InLooping )
+	{
+		Looping = InLooping;
+	}
+
 	/**
 	 * Changes the playback position of the media.
 	 *
@@ -202,6 +222,9 @@ private:
 	// Caches the last error result code (0 = no error).
 	HRESULT LastError;
 
+	// Whether playback should loop to the beginning.
+	bool Looping;
+
 	// The media session's clock.
 	TComPtr<IMFPresentationClock> PresentationClock;
 
diff --git a/Engine/Source/Runtime/Media/Public/Interfaces/IMediaPlayer.h b/Engine/Source/Runtime/Media/Public/Interfaces/IMediaPlayer.h
--- a/Engine/Source/Runtime/Media/Public/Interfaces/IMediaPlayer.h
+++ b/Engine/Source/Runtime/Media/Public/Interfaces/IMediaPlayer.h
@@ -92,6 +92,14 @@ public:
 	 */
 	virtual FString GetUrl( ) const = 0;
 
+	/**
+	 * Checks whether playback is currently looping.
+	 *
+	 * @return true if playback is looping, false otherwise.
+	 * @see SetLooping
+	 */
+	virtual bool IsLooping( ) const = 0;
+
 	/**
 	 * Checks whether media playback is currently paused.
 	 *
@@ -162,6 +170,14 @@ public:
 	 */
 	virtual bool Seek( const FTimespan& Time ) = 0;
 
+	/**
+	 * Sets whether playback should be looping.
+	 *
+	 * @param Looping Enables or disables looping.
+	 * @see IsLooping
+	 */
+	virtual bool SetLooping( bool Looping ) = 0;
+
 	/**
 	 * Checks whether the specified playback rate is supported.
 	 *
diff --git a/Engine/Source/Runtime/MediaAssets/Private/Assets/MediaAsset.cpp b/Engine/Source/Runtime/MediaAssets/Private/Assets/MediaAsset.cpp
--- a/Engine/Source/Runtime/MediaAssets/Private/Assets/MediaAsset.cpp
+++ b/Engine/Source/Runtime/MediaAssets/Private/Assets/MediaAsset.cpp
@@ -11,7 +11,6 @@ UMediaAsset::UMediaAsset( const class FPostConstructInitializeProperties& PCIP )
 	, AutoPlay(false)
 	, Looping(true)
 	, PlaybackRate(1.0f)
-	, ResetOnLastFrame(false)
 	, StreamMode(MASM_FromUrl)
 	, MediaPlayer(nullptr)
 { }
@@ -71,6 +70,18 @@ FTimespan UMediaAsset::GetTime( ) const
 }
 
 
+const FString& UMediaAsset::GetUrl( ) const
+{
+	return CurrentUrl;
+}
+
+
+bool UMediaAsset::IsLooping( ) const
+{
+	return MediaPlayer.IsValid() && MediaPlayer->IsLooping();
+}
+
+
 bool UMediaAsset::IsPaused( ) const
 {
 	return MediaPlayer.IsValid() && MediaPlayer->IsPaused();
@@ -89,21 +100,36 @@ bool UMediaAsset::IsStopped( ) const
 }
 
 
+bool UMediaAsset::OpenUrl( const FString& NewUrl )
+{
+	URL.FilePath = NewUrl;
+	InitializePlayer();
+
+	return (CurrentUrl == NewUrl);
+}
+
+
 bool UMediaAsset::Pause( )
 {
 	return MediaPlayer.IsValid() && MediaPlayer->Pause();
 }
 
 
+bool UMediaAsset::Play( )
+{
+	return MediaPlayer.IsValid() && MediaPlayer->Play();
+}
+
+
 bool UMediaAsset::Rewind( )
 {
 	return MediaPlayer.IsValid() && MediaPlayer->Seek(FTimespan::Zero());
 }
 
 
-bool UMediaAsset::Play( )
+bool UMediaAsset::SetLooping( bool InLooping )
 {
-	return MediaPlayer.IsValid() && MediaPlayer->Play();
+	return MediaPlayer.IsValid() && MediaPlayer->SetLooping(InLooping);
 }
 
 
@@ -279,6 +305,9 @@ void UMediaAsset::InitializePlayer( )
 	{
 		return;
 	}
+
+	// start playback, if desired
+	MediaPlayer->SetLooping(Looping);
     
 	if (AutoPlay)
 	{
diff --git a/Engine/Source/Runtime/MediaAssets/Public/MediaAsset.h b/Engine/Source/Runtime/MediaAssets/Public/MediaAsset.h
--- a/Engine/Source/Runtime/MediaAssets/Public/MediaAsset.h
+++ b/Engine/Source/Runtime/MediaAssets/Public/MediaAsset.h
@@ -40,34 +40,6 @@ class MEDIAASSETS_API UMediaAsset
 {
 	GENERATED_UCLASS_BODY()
 
-	/** Whether this media should automatically start playing when it is loaded. */
-	UPROPERTY(EditAnywhere, Category=Playback)
-	uint32 AutoPlay:1;
-
-	/** Whether playback should loop when it reaches the end. */
-	UPROPERTY(EditAnywhere, BlueprintReadWrite, Category=Playback)
-	uint32 Looping:1;
-
-	/** The media's playback rate (1.0 = real time) */
-	UPROPERTY(EditAnywhere, BlueprintReadOnly, Category=Playback)
-	float PlaybackRate;
-
-	/** The current playback position. */
-	UPROPERTY(BlueprintReadOnly, Category=Playback)
-	FTimespan Position;
-
-	/** Whether the media should reset to first frame when playback reaches the end. */
-	UPROPERTY(EditAnywhere, BlueprintReadWrite, Category=Playback)
-	uint32 ResetOnLastFrame:1;
-
-	/** Select where to stream the media from, i.e. file or memory. */
-	UPROPERTY(EditAnywhere, BlueprintReadOnly, Category=MediaSource)
-	TEnumAsByte<enum EMediaAssetStreamModes> StreamMode;
-
-	/** The URL to the media file to be played. */
-	UPROPERTY(EditAnywhere, BlueprintReadOnly, Category=MediaSource)
-	FFilePath URL;
-
 public:
 
 	/**
@@ -125,6 +97,24 @@ public:
 	UFUNCTION(BlueprintCallable, Category="Media|MediaAsset")
 	FTimespan GetTime( ) const;
 
+	/**
+	 * Gets the URL of the currently loaded media, if any.
+	 *
+	 * @return Media URL, or empty string if no media was loaded.
+	 * @see OpenUrl
+	 */
+	UFUNCTION(BlueprintCallable, Category="Media|MediaAsset")
+	const FString& GetUrl( ) const;
+
+	/**
+	 * Checks whether playback is looping.
+	 *
+	 * @return true if looping, false otherwise.
+	 * @see SetLooping
+	 */
+	UFUNCTION(BlueprintCallable, Category="Media|MediaAsset")
+	bool IsLooping( ) const;
+
 	/**
 	 * Checks whether playback is currently paused.
 	 *
@@ -152,6 +142,16 @@ public:
 	UFUNCTION(BlueprintCallable, Category="Media|MediaAsset")
 	bool IsStopped( ) const;
 
+	/**
+	 * Opens the specified media URL.
+	 *
+	 * @param NewUrl The URL to open.
+	 * @return true on success, false otherwise.
+	 * @see GetUrl
+	 */
+	UFUNCTION(BlueprintCallable, Category="Media|MediaAsset")
+	bool OpenUrl( const FString& NewUrl );
+
 	/**
 	 * Pauses media playback.
 	 *
@@ -189,6 +189,16 @@ public:
 	UFUNCTION(BlueprintCallable, Category="Media|MediaAsset")
 	bool Seek( const FTimespan& InTime );
 
+	/**
+	 * Enables or disables playback looping.
+	 *
+	 * @param Looping Whether playback should be looped.
+	 * @return true on success, false otherwise.
+	 * @see IsLooping
+	 */
+	UFUNCTION(BlueprintCallable, Category="Media|MediaAsset")
+	bool SetLooping( bool InLooping );
+
 	/**
 	 * Changes the media's playback rate.
 	 *
@@ -274,6 +284,28 @@ protected:
 	/** Initializes the media player. */
 	void InitializePlayer( );
 
+protected:
+
+	/** Whether this media should automatically start playing when it is loaded. */
+	UPROPERTY(EditAnywhere, Category=Playback)
+	uint32 AutoPlay:1;
+
+	/** Whether playback should loop when it reaches the end. */
+	UPROPERTY(EditAnywhere, Category=Playback)
+	uint32 Looping:1;
+
+	/** The media's playback rate (1.0 = real time) */
+	UPROPERTY(EditAnywhere, Category=Playback)
+	float PlaybackRate;
+
+	/** Select where to stream the media from, i.e. file or memory. */
+	UPROPERTY(EditAnywhere, BlueprintReadOnly, Category=MediaSource)
+	TEnumAsByte<enum EMediaAssetStreamModes> StreamMode;
+
+	/** The URL to the media file to be played. */
+	UPROPERTY(EditAnywhere, Category=MediaSource)
+	FFilePath URL;
+
 private:
 
 	// Callback for when the media player is unloading media.
