Fixed build on case-sensitive fs.

diff --git a/Engine/Source/Runtime/MediaAssets/Private/MediaAssetsPrivatePCH.h b/Engine/Source/Runtime/MediaAssets/Private/MediaAssetsPrivatePCH.h
--- a/Engine/Source/Runtime/MediaAssets/Private/MediaAssetsPrivatePCH.h
+++ b/Engine/Source/Runtime/MediaAssets/Private/MediaAssetsPrivatePCH.h
@@ -11,7 +11,7 @@
 #include "Engine.h"
 #include "Media.h"
 
-#include "MaterialExpressionTExtureSampleParameterMedia.h"
+#include "MaterialExpressionTextureSampleParameterMedia.h"
 #include "MediaAsset.h"
 #include "MediaSoundWave.h"
 #include "MediaTexture.h"
