Media: cosmetic changes

diff --git a/Engine/Source/Runtime/MediaAssets/Private/Assets/MediaTextureResource.cpp b/Engine/Source/Runtime/MediaAssets/Private/Assets/MediaTextureResource.cpp
--- a/Engine/Source/Runtime/MediaAssets/Private/Assets/MediaTextureResource.cpp
+++ b/Engine/Source/Runtime/MediaAssets/Private/Assets/MediaTextureResource.cpp
@@ -1,7 +1,7 @@
 // Copyright 1998-2014 Epic Games, Inc. All Rights Reserved.
 
 #include "MediaAssetsPrivatePCH.h"
-#include "Assets/MediaTextureResource.h"
+
 
 /* FMediaTextureResource structors
  *****************************************************************************/
diff --git a/Engine/Source/Runtime/MediaAssets/Private/MediaAssetsPrivatePCH.h b/Engine/Source/Runtime/MediaAssets/Private/MediaAssetsPrivatePCH.h
--- a/Engine/Source/Runtime/MediaAssets/Private/MediaAssetsPrivatePCH.h
+++ b/Engine/Source/Runtime/MediaAssets/Private/MediaAssetsPrivatePCH.h
@@ -15,3 +15,9 @@
 #include "MediaAsset.h"
 #include "MediaSoundWave.h"
 #include "MediaTexture.h"
+
+
+/* Private includes
+ *****************************************************************************/
+
+#include "MediaTextureResource.h"
