Media: various improvements

- slightly changed MediaAsset API for Play, Pause and SetRate
- exposed MediaAssets as BP variables
- better logic for MediaAsset toolbar buttons

diff --git a/Engine/Plugins/Media/MediaAssetEditor/Source/MediaAssetEditor/Private/MediaAssetEditorToolkit.cpp b/Engine/Plugins/Media/MediaAssetEditor/Source/MediaAssetEditor/Private/MediaAssetEditorToolkit.cpp
--- a/Engine/Plugins/Media/MediaAssetEditor/Source/MediaAssetEditor/Private/MediaAssetEditorToolkit.cpp
+++ b/Engine/Plugins/Media/MediaAssetEditor/Source/MediaAssetEditor/Private/MediaAssetEditorToolkit.cpp
@@ -266,7 +266,7 @@ float FMediaAssetEditorToolkit::GetReverseRate( ) const
 
 	if (Rate > -1.0f)
 	{
-		Rate = -1.0f;
+		return -1.0f;
 	}
 
 	return 2.0f * Rate;
@@ -302,7 +302,7 @@ void FMediaAssetEditorToolkit::HandlePauseMediaActionExecute( )
 
 bool FMediaAssetEditorToolkit::HandlePlayMediaActionCanExecute( ) const
 {
-	return MediaAsset->CanPlay() && !FMath::IsNearlyEqual(MediaAsset->GetRate(), 1.0f);
+	return MediaAsset->CanPlay() && (MediaAsset->GetRate() != 1.0f);
 }
 
 
diff --git a/Engine/Plugins/Media/WmfMedia/Source/WmfMedia/Private/Player/WmfMediaPlayer.cpp b/Engine/Plugins/Media/WmfMedia/Source/WmfMedia/Private/Player/WmfMediaPlayer.cpp
--- a/Engine/Plugins/Media/WmfMedia/Source/WmfMedia/Private/Player/WmfMediaPlayer.cpp
+++ b/Engine/Plugins/Media/WmfMedia/Source/WmfMedia/Private/Player/WmfMediaPlayer.cpp
@@ -171,43 +171,43 @@ bool FWmfMediaPlayer::Open( const TSharedRef<TArray<uint8>>& Buffer, const FStri
 }
 
 
-bool FWmfMediaPlayer::Play( float Rate )
+bool FWmfMediaPlayer::Seek( const FTimespan& Time )
+{
+	return ((MediaSession != NULL) && MediaSession->SetPosition(Time));
+}
+
+
+bool FWmfMediaPlayer::SetLooping( bool Looping )
 {
 	if (MediaSession == NULL)
 	{
 		return false;
 	}
 
-	if (FMath::IsNearlyZero(Rate))
-	{
-		if (MediaSession->SupportsScrubbing())
-		{
-			return MediaSession->SetRate(Rate);
-		}
-
-		return MediaSession->SetState(EMediaStates::Paused);
-	}
-
-	return MediaSession->SetRate(Rate) && MediaSession->SetState(EMediaStates::Playing);
-}
-
+	MediaSession->SetLooping(Looping);
 
-bool FWmfMediaPlayer::Seek( const FTimespan& Time )
-{
-	return ((MediaSession != NULL) && MediaSession->SetPosition(Time));
+	return true;
 }
 
 
-bool FWmfMediaPlayer::SetLooping( bool Looping )
+bool FWmfMediaPlayer::SetRate( float Rate )
 {
 	if (MediaSession == NULL)
 	{
 		return false;
 	}
 
-	MediaSession->SetLooping(Looping);
+	if (FMath::IsNearlyZero(Rate))
+	{
+		if (MediaSession->SupportsScrubbing())
+		{
+			return MediaSession->SetRate(Rate);
+		}
 
-	return true;
+		return MediaSession->SetState(EMediaStates::Paused);
+	}
+
+	return MediaSession->SetRate(Rate) && MediaSession->SetState(EMediaStates::Playing);
 }
 
 
diff --git a/Engine/Plugins/Media/WmfMedia/Source/WmfMedia/Private/Player/WmfMediaPlayer.h b/Engine/Plugins/Media/WmfMedia/Source/WmfMedia/Private/Player/WmfMediaPlayer.h
--- a/Engine/Plugins/Media/WmfMedia/Source/WmfMedia/Private/Player/WmfMediaPlayer.h
+++ b/Engine/Plugins/Media/WmfMedia/Source/WmfMedia/Private/Player/WmfMediaPlayer.h
@@ -39,9 +39,9 @@ public:
 	virtual bool IsReady( ) const override;
 	virtual bool Open( const FString& Url ) override;
 	virtual bool Open( const TSharedRef<TArray<uint8>>& Buffer, const FString& OriginalUrl ) override;
-	virtual bool Play( float Rate ) override;
 	virtual bool Seek( const FTimespan& Time ) override;
 	virtual bool SetLooping( bool Looping ) override;
+	virtual bool SetRate( float Rate ) override;
 	virtual bool SupportsRate( float Rate, bool Unthinned ) const override;
 	virtual bool SupportsScrubbing( ) const override;
 	virtual bool SupportsSeeking( ) const override;
diff --git a/Engine/Source/Runtime/Media/Public/Interfaces/IMediaPlayer.h b/Engine/Source/Runtime/Media/Public/Interfaces/IMediaPlayer.h
--- a/Engine/Source/Runtime/Media/Public/Interfaces/IMediaPlayer.h
+++ b/Engine/Source/Runtime/Media/Public/Interfaces/IMediaPlayer.h
@@ -66,7 +66,7 @@ public:
 	 * Gets the nominal playback rate, i.e. 1.0 for real time.
 	 *
 	 * @return Playback rate.
-	 * @see Play
+	 * @see Pause, Play, SetRate, SupportsRate
 	 */
 	virtual float GetRate( ) const = 0;
 
@@ -147,20 +147,6 @@ public:
 	 */
 	virtual bool Open( const TSharedRef<TArray<uint8>>& Buffer, const FString& OriginalUrl ) = 0;
 
-	/**
-	 * Starts media playback at the specified rate.
-	 *
-	 * A playback rate of 1.0 will play the media normally at real-time.
-	 * A rate of 0.0 corresponds to pausing playback. A negative rate, if
-	 * supported, plays the media in reverse, and a rate larger than 1.0
-	 * fast forwards playback.
-	 *
-	 * @param Rate The playback rate to use.
-	 * @return true if the media will be played, false otherwise.
-	 * @see GetRate, Pause, Stop
-	 */
-	virtual bool Play( float Rate ) = 0;
-
 	/**
 	 * Changes the media's playback time.
 	 *
@@ -178,13 +164,27 @@ public:
 	 */
 	virtual bool SetLooping( bool Looping ) = 0;
 
+	/**
+	 * Sets the current playback rate.
+	 *
+	 * A playback rate of 1.0 will play the media normally at real-time.
+	 * A rate of 0.0 corresponds to pausing playback. A negative rate, if
+	 * supported, plays the media in reverse, and a rate larger than 1.0
+	 * fast forwards playback.
+	 *
+	 * @param Rate The playback rate to set.
+	 * @return true on success, false otherwise.
+	 * @see GetRate, Pause, Play, SupportsRate
+	 */
+	virtual bool SetRate( float Rate ) = 0;
+
 	/**
 	 * Checks whether the specified playback rate is supported.
 	 *
 	 * @param Rate The rate to check (can be negative for reverse play).
 	 * @param Unthinned Whether no frames should be dropped at the given rate.
 	 * @return true if the rate is supported, false otherwise.
-	 * @see SupportsScrubbing, SupportsSeeking
+	 * @see GetRate, SetRate, SupportsScrubbing, SupportsSeeking
 	 */
 	virtual bool SupportsRate( float Rate, bool Unthinned ) const = 0;
 
@@ -243,25 +243,27 @@ public:
 	/**
 	 * Pauses media playback.
 	 *
+	 * This is the same as setting the playback rate to 0.0.
+	 *
 	 * @return true if the media is being paused, false otherwise.
 	 * @see Play, Stop
 	 */
 	FORCEINLINE bool Pause( )
 	{
-		return Play(0.0f);
+		return SetRate(0.0f);
 	}
 
 	/**
 	 * Starts media playback at the default rate of 1.0.
 	 *
-	 * This method is equivalent to setting the playback rate to 1.0.
+	 * This is the same as setting the playback rate to 1.0.
 	 *
 	 * @return true if playback is starting, false otherwise.
 	 * @see Pause, Stop
 	 */
 	FORCEINLINE bool Play( )
 	{
-		return Play(1.0f);
+		return SetRate(1.0f);
 	}
 
 	/**
diff --git a/Engine/Source/Runtime/MediaAssets/Private/Assets/MediaAsset.cpp b/Engine/Source/Runtime/MediaAssets/Private/Assets/MediaAsset.cpp
--- a/Engine/Source/Runtime/MediaAssets/Private/Assets/MediaAsset.cpp
+++ b/Engine/Source/Runtime/MediaAssets/Private/Assets/MediaAsset.cpp
@@ -10,7 +10,6 @@ UMediaAsset::UMediaAsset( const class FPostConstructInitializeProperties& PCIP )
 	: Super(PCIP)
 	, AutoPlay(false)
 	, Looping(true)
-	, PlaybackRate(1.0f)
 	, StreamMode(MASM_FromUrl)
 	, MediaPlayer(nullptr)
 { }
@@ -27,7 +26,7 @@ bool UMediaAsset::CanPause( ) const
 
 bool UMediaAsset::CanPlay( ) const
 {
-	return MediaPlayer.IsValid() && !MediaPlayer->IsPlaying() && MediaPlayer->IsReady();
+	return MediaPlayer.IsValid() && MediaPlayer->IsReady();
 }
 
 
@@ -105,19 +104,19 @@ bool UMediaAsset::OpenUrl( const FString& NewUrl )
 
 bool UMediaAsset::Pause( )
 {
-	return MediaPlayer.IsValid() && MediaPlayer->Pause();
+	return SetRate(0.0f);
 }
 
 
 bool UMediaAsset::Play( )
 {
-	return MediaPlayer.IsValid() && MediaPlayer->Play();
+	return SetRate(1.0f);
 }
 
 
 bool UMediaAsset::Rewind( )
 {
-	return MediaPlayer.IsValid() && MediaPlayer->Seek(FTimespan::Zero());
+	return Seek(FTimespan::Zero());
 }
 
 
@@ -129,14 +128,7 @@ bool UMediaAsset::SetLooping( bool InLooping )
 
 bool UMediaAsset::SetRate( float InRate )
 {
-	PlaybackRate = InRate;
-
-	if (!MediaPlayer.IsValid() || !MediaPlayer->IsPlaying())
-	{
-		return false;
-	}
-
-	return MediaPlayer->Play(InRate);
+	return MediaPlayer.IsValid() && MediaPlayer->SetRate(InRate);
 }
 
 
diff --git a/Engine/Source/Runtime/MediaAssets/Public/MediaAsset.h b/Engine/Source/Runtime/MediaAssets/Public/MediaAsset.h
--- a/Engine/Source/Runtime/MediaAssets/Public/MediaAsset.h
+++ b/Engine/Source/Runtime/MediaAssets/Public/MediaAsset.h
@@ -34,7 +34,7 @@ DECLARE_DYNAMIC_MULTICAST_DELEGATE_OneParam(FOnMediaAssetMediaOpened, FString, O
  * This class is represents a media URL along with a corresponding media player
  * for exposing media playback functionality to the Engine and to Blueprints.
  */
-UCLASS(hidecategories=(Object))
+UCLASS(BlueprintType, hidecategories=(Object))
 class MEDIAASSETS_API UMediaAsset
 	: public UObject
 {
@@ -66,7 +66,7 @@ public:
 	 * Gets the media's duration.
 	 *
 	 * @return A time span representing the duration.
-	 * @see GetPosition
+	 * @see GetTime, Seek
 	 */
 	UFUNCTION(BlueprintCallable, Category="Media|MediaAsset")
 	FTimespan GetDuration( ) const;
@@ -75,6 +75,7 @@ public:
 	 * Gets the media's current playback rate.
 	 *
 	 * @return The playback rate.
+	 * @see SetRate, SupportsRate
 	 */
 	UFUNCTION(BlueprintCallable, Category="Media|MediaAsset")
 	float GetRate( ) const;
@@ -146,36 +147,42 @@ public:
 	/**
 	 * Pauses media playback.
 	 *
+	 * This is the same as setting the playback rate to 0.0.
+	 *
 	 * @return true if playback is being paused, false otherwise.
-	 * @see CanPause, Play, Rewind, Seek, Stop
+	 * @see CanPause, Play, Rewind, Seek, SetRate
 	 */
 	UFUNCTION(BlueprintCallable, Category="Media|MediaAsset")
 	bool Pause( );
-	
+
 	/**
 	 * Starts media playback.
 	 *
-	 * @return true if rewinding, false otherwise.
-	 * @see GetTime, Pause, Play, Seek, Stop
+	 * This is the same as setting the playback rate to 1.0.
+	 *
+	 * @return true if playback is starting, false otherwise.
+	 * @see CanPlay, Pause, Rewind, Seek, SetRate
 	 */
 	UFUNCTION(BlueprintCallable, Category="Media|MediaAsset")
-	bool Rewind( );
+	bool Play( );
 
 	/**
-	 * Rewinds the media.
+	 * Rewinds the media to the beginning.
 	 *
-	 * @return true if playback is starting, false otherwise.
-	 * @see CanPlay, Pause, Rewind, Seek, Stop
+	 * This is the same as seeking to zero time.
+	 *
+	 * @return true if rewinding, false otherwise.
+	 * @see GetTime, Pause, Play, Seek
 	 */
 	UFUNCTION(BlueprintCallable, Category="Media|MediaAsset")
-	bool Play( );
+	bool Rewind( );
 
 	/**
 	 * Seeks to the specified playback time.
 	 *
 	 * @param InTime The playback time to set.
 	 * @return true on success, false otherwise.
-	 * @see GetTime
+	 * @see GetTime, Rewind
 	 */
 	UFUNCTION(BlueprintCallable, Category="Media|MediaAsset")
 	bool Seek( const FTimespan& InTime );
@@ -195,7 +202,7 @@ public:
 	 *
 	 * @param Rate The playback rate to set.
 	 * @return true on success, false otherwise.
-	 * @see GetRate
+	 * @see GetRate, SupportsRate
 	 */
 	UFUNCTION(BlueprintCallable, Category="Media|MediaAsset")
 	bool SetRate( float InRate );
@@ -277,10 +284,6 @@ protected:
 	UPROPERTY(EditAnywhere, Category=Playback)
 	uint32 Looping:1;
 
-	/** The media's playback rate (1.0 = real time) */
-	UPROPERTY(EditAnywhere, Category=Playback)
-	float PlaybackRate;
-
 	/** Select where to stream the media from, i.e. file or memory. */
 	UPROPERTY(EditAnywhere, BlueprintReadOnly, Category=MediaSource)
 	TEnumAsByte<enum EMediaAssetStreamModes> StreamMode;
