Media: Added auto play default playback rate property

diff --git a/Engine/Source/Runtime/MediaAssets/Private/Assets/MediaAsset.cpp b/Engine/Source/Runtime/MediaAssets/Private/Assets/MediaAsset.cpp
--- a/Engine/Source/Runtime/MediaAssets/Private/Assets/MediaAsset.cpp
+++ b/Engine/Source/Runtime/MediaAssets/Private/Assets/MediaAsset.cpp
@@ -9,6 +9,7 @@
 UMediaAsset::UMediaAsset( const class FPostConstructInitializeProperties& PCIP )
 	: Super(PCIP)
 	, AutoPlay(false)
+	, AutoPlayRate(1.0f)
 	, Looping(true)
 	, StreamMode(MASM_FromUrl)
 	, MediaPlayer(nullptr)
@@ -287,7 +288,7 @@ void UMediaAsset::InitializePlayer( )
     
 	if (AutoPlay)
 	{
-		MediaPlayer->Play();
+		MediaPlayer->SetRate(AutoPlayRate);
 	}
 }
 
diff --git a/Engine/Source/Runtime/MediaAssets/Public/MediaAsset.h b/Engine/Source/Runtime/MediaAssets/Public/MediaAsset.h
--- a/Engine/Source/Runtime/MediaAssets/Public/MediaAsset.h
+++ b/Engine/Source/Runtime/MediaAssets/Public/MediaAsset.h
@@ -280,6 +280,10 @@ protected:
 	UPROPERTY(EditAnywhere, Category=Playback)
 	uint32 AutoPlay:1;
 
+	/** The default playback rate when auto play is enabled. */
+	UPROPERTY(EditAnywhere, Category=Playback)
+	float AutoPlayRate;
+
 	/** Whether playback should loop when it reaches the end. */
 	UPROPERTY(EditAnywhere, Category=Playback)
 	uint32 Looping:1;
