Media: Removed material expressions for media textures (no longer needed)

diff --git a/Engine/Source/Runtime/MediaAssets/Private/Assets/MaterialExpressionTextureSampleParameterMedia.cpp b/Engine/Source/Runtime/MediaAssets/Private/Assets/MaterialExpressionTextureSampleParameterMedia.cpp
deleted file mode 100644
diff --git a/Engine/Source/Runtime/MediaAssets/Private/MediaAssetsPrivatePCH.h b/Engine/Source/Runtime/MediaAssets/Private/MediaAssetsPrivatePCH.h
--- a/Engine/Source/Runtime/MediaAssets/Private/MediaAssetsPrivatePCH.h
+++ b/Engine/Source/Runtime/MediaAssets/Private/MediaAssetsPrivatePCH.h
@@ -10,8 +10,6 @@
 #include "CoreUObject.h"
 #include "Engine.h"
 #include "Media.h"
-
-#include "MaterialExpressionTextureSampleParameterMedia.h"
 #include "MediaAsset.h"
 #include "MediaSoundWave.h"
 #include "MediaTexture.h"
diff --git a/Engine/Source/Runtime/MediaAssets/Public/MaterialExpressionTextureSampleParameterMedia.h b/Engine/Source/Runtime/MediaAssets/Public/MaterialExpressionTextureSampleParameterMedia.h
deleted file mode 100644
