Media: cosmetic changes

diff --git a/Engine/Source/Runtime/MediaAssets/Private/Assets/MediaAsset.cpp b/Engine/Source/Runtime/MediaAssets/Private/Assets/MediaAsset.cpp
--- a/Engine/Source/Runtime/MediaAssets/Private/Assets/MediaAsset.cpp
+++ b/Engine/Source/Runtime/MediaAssets/Private/Assets/MediaAsset.cpp
@@ -49,7 +49,7 @@ float UMediaAsset::GetRate( ) const
 		return MediaPlayer->GetRate();
 	}
 
-	return 0;
+	return 0.0f;
 }
 
 
@@ -127,9 +127,9 @@ bool UMediaAsset::SetLooping( bool InLooping )
 }
 
 
-bool UMediaAsset::SetRate( float InRate )
+bool UMediaAsset::SetRate( float Rate )
 {
-	return MediaPlayer.IsValid() && MediaPlayer->SetRate(InRate);
+	return MediaPlayer.IsValid() && MediaPlayer->SetRate(Rate);
 }
 
 
diff --git a/Engine/Source/Runtime/MediaAssets/Public/MediaAsset.h b/Engine/Source/Runtime/MediaAssets/Public/MediaAsset.h
--- a/Engine/Source/Runtime/MediaAssets/Public/MediaAsset.h
+++ b/Engine/Source/Runtime/MediaAssets/Public/MediaAsset.h
@@ -205,7 +205,7 @@ public:
 	 * @see GetRate, SupportsRate
 	 */
 	UFUNCTION(BlueprintCallable, Category="Media|MediaAsset")
-	bool SetRate( float InRate );
+	bool SetRate( float Rate );
 
 	/**
 	 * Checks whether the specified playback rate is supported.
