Media: Moved implementation into .cpp to make UMediaAsset truly forward declarable

diff --git a/Engine/Source/Runtime/MediaAssets/Private/Assets/MediaTexture.cpp b/Engine/Source/Runtime/MediaAssets/Private/Assets/MediaTexture.cpp
--- a/Engine/Source/Runtime/MediaAssets/Private/Assets/MediaTexture.cpp
+++ b/Engine/Source/Runtime/MediaAssets/Private/Assets/MediaTexture.cpp
@@ -32,6 +32,14 @@ UMediaTexture::~UMediaTexture( )
 /* UMediaTexture interface
  *****************************************************************************/
 
+TSharedPtr<class IMediaPlayer> UMediaTexture::GetMediaPlayer( ) const
+{
+	return (MediaAsset != nullptr)
+		? MediaAsset->GetMediaPlayer()
+		: nullptr;
+}
+
+
 void UMediaTexture::SetMediaAsset( UMediaAsset* InMediaAsset )
 {
 	MediaAsset = InMediaAsset;
diff --git a/Engine/Source/Runtime/MediaAssets/Public/MediaTexture.h b/Engine/Source/Runtime/MediaAssets/Public/MediaTexture.h
--- a/Engine/Source/Runtime/MediaAssets/Public/MediaTexture.h
+++ b/Engine/Source/Runtime/MediaAssets/Public/MediaTexture.h
@@ -38,16 +38,6 @@ public:
 	/** Destructor. */
 	~UMediaTexture( );
 
-public:
-
-	/**
-	 * Sets the media asset to be used for this texture.
-	 *
-	 * @param InMediaAsset The asset to set.
-	 */
-	UFUNCTION(BlueprintCallable, Category="Media|MediaTexture")
-	void SetMediaAsset( UMediaAsset* InMediaAsset );
-
 public:
 
 	/**
@@ -75,10 +65,7 @@ public:
 	 *
 	 * @return Media player, or nullptr if no player is available.
 	 */
-	TSharedPtr<class IMediaPlayer> GetMediaPlayer( ) const
-	{
-		return (MediaAsset != nullptr) ? MediaAsset->GetMediaPlayer() : nullptr;
-	}
+	TSharedPtr<class IMediaPlayer> GetMediaPlayer( ) const;
 
 	/**
 	 * Gets the currently selected video track, if any.
@@ -90,6 +77,14 @@ public:
 		return VideoTrack;
 	}
 	
+	/**
+	 * Sets the media asset to be used for this texture.
+	 *
+	 * @param InMediaAsset The asset to set.
+	 */
+	UFUNCTION(BlueprintCallable, Category="Media|MediaTexture")
+	void SetMediaAsset( UMediaAsset* InMediaAsset );
+
 public:
 
 	// UTexture overrides.
