Media: Properly initializing media assets on startup and when media is loaded/unloaded

diff --git a/Engine/Plugins/Media/WmfMedia/Source/WmfMedia/Private/Player/WmfMediaPlayer.cpp b/Engine/Plugins/Media/WmfMedia/Source/WmfMedia/Private/Player/WmfMediaPlayer.cpp
--- a/Engine/Plugins/Media/WmfMedia/Source/WmfMedia/Private/Player/WmfMediaPlayer.cpp
+++ b/Engine/Plugins/Media/WmfMedia/Source/WmfMedia/Private/Player/WmfMediaPlayer.cpp
@@ -72,21 +72,20 @@ bool FWmfMediaPlayer::SupportsSeeking( ) const
 
 void FWmfMediaPlayer::Close( )
 {
-	if (!MediaUrl.IsEmpty())
+	if (MediaSession == NULL)
 	{
-		ClosingEvent.Broadcast(MediaUrl);
+		return;
 	}
 
-	if (MediaSession != NULL)
-	{
-		MediaSession->SetState(EMediaStates::Closed);
-		MediaSession.Reset();
-	}
+	MediaSession->SetState(EMediaStates::Closed);
+	MediaSession.Reset();
 
 	Tracks.Reset();
 
 	Duration = 0;
 	MediaUrl = FString();
+
+	ClosedEvent.Broadcast();
 }
 
 
diff --git a/Engine/Plugins/Media/WmfMedia/Source/WmfMedia/Private/Player/WmfMediaPlayer.h b/Engine/Plugins/Media/WmfMedia/Source/WmfMedia/Private/Player/WmfMediaPlayer.h
--- a/Engine/Plugins/Media/WmfMedia/Source/WmfMedia/Private/Player/WmfMediaPlayer.h
+++ b/Engine/Plugins/Media/WmfMedia/Source/WmfMedia/Private/Player/WmfMediaPlayer.h
@@ -54,10 +54,10 @@ public:
 	virtual bool SetLooping( bool Looping ) override;
 	virtual bool SetRate( float Rate ) override;
 
-	DECLARE_DERIVED_EVENT(FWmfMediaPlayer, IMediaPlayer::FOnMediaClosing, FOnMediaClosing);
-	virtual FOnMediaClosing& OnClosing( ) override
+	DECLARE_DERIVED_EVENT(FWmfMediaPlayer, IMediaPlayer::FOnMediaClosed, FOnMediaClosed);
+	virtual FOnMediaClosed& OnClosed( ) override
 	{
-		return ClosingEvent;
+		return ClosedEvent;
 	}
 
 	DECLARE_DERIVED_EVENT(FWmfMediaPlayer, IMediaPlayer::FOnMediaOpened, FOnMediaOpened);
@@ -90,24 +90,24 @@ protected:
 
 private:
 
-	// The available media tracks.
+	/** The available media tracks. */
 	TArray<IMediaTrackRef> Tracks;
 
-	// The duration of the currently loaded media.
+	/** The duration of the currently loaded media. */
 	FTimespan Duration;
 
-	// Holds the asynchronous callback object for the media stream.
+	/** Holds the asynchronous callback object for the media stream. */
 	TComPtr<FWmfMediaSession> MediaSession;
 
-	// The URL of the currently opened media.
+	/** The URL of the currently opened media. */
 	FString MediaUrl;
 
 private:
 
-	// Holds an event delegate that is invoked when media is being unloaded.
-	FOnMediaClosing ClosingEvent;
+	/** Holds an event delegate that is invoked when media has been closed. */
+	FOnMediaClosed ClosedEvent;
 
-	// Holds an event delegate that is invoked when media has been loaded.
+	/** Holds an event delegate that is invoked when media has been opened. */
 	FOnMediaOpened OpenedEvent;
 };
 
diff --git a/Engine/Source/Runtime/Media/Public/Interfaces/IMediaPlayer.h b/Engine/Source/Runtime/Media/Public/Interfaces/IMediaPlayer.h
--- a/Engine/Source/Runtime/Media/Public/Interfaces/IMediaPlayer.h
+++ b/Engine/Source/Runtime/Media/Public/Interfaces/IMediaPlayer.h
@@ -173,8 +173,8 @@ public:
 public:
 
 	/** Gets an event delegate that is invoked when media has been closed. */
-	DECLARE_EVENT_OneParam(IMediaPlayer, FOnMediaClosing, FString /*ClosedUrl*/)
-	virtual FOnMediaClosing& OnClosing( ) = 0;
+	DECLARE_EVENT(IMediaPlayer, FOnMediaClosed)
+	virtual FOnMediaClosed& OnClosed( ) = 0;
 
 	/** Gets an event delegate that is invoked when media has been opened. */
 	DECLARE_EVENT_OneParam(IMediaPlayer, FOnMediaOpened, FString /*OpenedUrl*/)
diff --git a/Engine/Source/Runtime/MediaAssets/Private/Assets/MediaAsset.cpp b/Engine/Source/Runtime/MediaAssets/Private/Assets/MediaAsset.cpp
--- a/Engine/Source/Runtime/MediaAssets/Private/Assets/MediaAsset.cpp
+++ b/Engine/Source/Runtime/MediaAssets/Private/Assets/MediaAsset.cpp
@@ -219,7 +219,7 @@ void UMediaAsset::InitializePlayer( )
 		if (MediaPlayer.IsValid())
 		{
 			MediaPlayer->Close();
-			MediaPlayer->OnClosing().RemoveAll(this);
+			MediaPlayer->OnClosed().RemoveAll(this);
 			MediaPlayer->OnOpened().RemoveAll(this);
 			MediaPlayer.Reset();
 		}
@@ -237,7 +237,7 @@ void UMediaAsset::InitializePlayer( )
 			return;
 		}
 
-		MediaPlayer->OnClosing().AddUObject(this, &UMediaAsset::HandleMediaPlayerMediaClosing);
+		MediaPlayer->OnClosed().AddUObject(this, &UMediaAsset::HandleMediaPlayerMediaClosed);
 		MediaPlayer->OnOpened().AddUObject(this, &UMediaAsset::HandleMediaPlayerMediaOpened);
 
 		// open the new media file
@@ -294,13 +294,15 @@ void UMediaAsset::InitializePlayer( )
 /* UMediaAsset callbacks
  *****************************************************************************/
 
-void UMediaAsset::HandleMediaPlayerMediaClosing( FString ClosedUrl )
+void UMediaAsset::HandleMediaPlayerMediaClosed( )
 {
-	OnMediaClosing.Broadcast(ClosedUrl);
+	MediaChangedEvent.Broadcast();
+	OnMediaClosed.Broadcast();
 }
 
 
 void UMediaAsset::HandleMediaPlayerMediaOpened( FString OpenedUrl )
 {
+	MediaChangedEvent.Broadcast();
 	OnMediaOpened.Broadcast(OpenedUrl);
 }
diff --git a/Engine/Source/Runtime/MediaAssets/Private/Assets/MediaSoundWave.cpp b/Engine/Source/Runtime/MediaAssets/Private/Assets/MediaSoundWave.cpp
--- a/Engine/Source/Runtime/MediaAssets/Private/Assets/MediaSoundWave.cpp
+++ b/Engine/Source/Runtime/MediaAssets/Private/Assets/MediaSoundWave.cpp
@@ -43,7 +43,7 @@ void UMediaSoundWave::SetMediaAsset( UMediaAsset* InMediaAsset )
 {
 	MediaAsset = InMediaAsset;
 
-	InitializeStream();
+	InitializeTrack();
 }
 
 
@@ -127,33 +127,21 @@ void UMediaSoundWave::Serialize( FArchive& Ar )
 /* UMediaSoundWave implementation
  *****************************************************************************/
 
-void UMediaSoundWave::InitializeStream( )
+void UMediaSoundWave::InitializeTrack( )
 {
 	// assign new media asset
 	if (CurrentMediaAsset != MediaAsset)
 	{
 		if (CurrentMediaAsset != nullptr)
 		{
-			IMediaPlayerPtr MediaPlayer = CurrentMediaAsset->GetMediaPlayer();
-
-			if (MediaPlayer.IsValid())
-			{
-				MediaPlayer->OnClosing().RemoveAll(this);
-				MediaPlayer->OnOpened().RemoveAll(this);			
-			}
+			CurrentMediaAsset->OnMediaChanged().RemoveAll(this);
 		}
 
 		CurrentMediaAsset = MediaAsset;
 
 		if (MediaAsset != nullptr)
 		{
-			IMediaPlayerPtr MediaPlayer = MediaAsset->GetMediaPlayer();
-
-			if (MediaPlayer.IsValid())
-			{
-				MediaPlayer->OnClosing().AddUObject(this, &UMediaSoundWave::HandleMediaPlayerClosing);
-				MediaPlayer->OnOpened().AddUObject(this, &UMediaSoundWave::HandleMediaPlayerOpened);
-			}
+			MediaAsset->OnMediaChanged().AddUObject(this, &UMediaSoundWave::HandleMediaAssetMediaChanged);
 		}	
 	}
 
@@ -198,13 +186,7 @@ void UMediaSoundWave::InitializeStream( )
 /* UMediaSoundWave callbacks
  *****************************************************************************/
 
-void UMediaSoundWave::HandleMediaPlayerClosing( FString ClosingUrl )
-{
-	InitializeStream();
-}
-
-
-void UMediaSoundWave::HandleMediaPlayerOpened( FString OpenedUrl )
+void UMediaSoundWave::HandleMediaAssetMediaChanged( )
 {
-	InitializeStream();
+	InitializeTrack();
 }
diff --git a/Engine/Source/Runtime/MediaAssets/Private/Assets/MediaTexture.cpp b/Engine/Source/Runtime/MediaAssets/Private/Assets/MediaTexture.cpp
--- a/Engine/Source/Runtime/MediaAssets/Private/Assets/MediaTexture.cpp
+++ b/Engine/Source/Runtime/MediaAssets/Private/Assets/MediaTexture.cpp
@@ -157,6 +157,8 @@ void UMediaTexture::PostEditChangeProperty( FPropertyChangedEvent& PropertyChang
 	Super::PostEditChangeProperty(PropertyChangedEvent);
 }
 
+#endif // WITH_EDITOR
+
 
 /* UMediaTexture implementation
  *****************************************************************************/
@@ -168,26 +170,14 @@ void UMediaTexture::InitializeTrack( )
 	{
 		if (CurrentMediaAsset != nullptr)
 		{
-			IMediaPlayerPtr MediaPlayer = CurrentMediaAsset->GetMediaPlayer();
-
-			if (MediaPlayer.IsValid())
-			{
-				MediaPlayer->OnClosing().RemoveAll(this);
-				MediaPlayer->OnOpened().RemoveAll(this);
-			}
+			CurrentMediaAsset->OnMediaChanged().RemoveAll(this);
 		}
 
 		CurrentMediaAsset = MediaAsset;
 
 		if (MediaAsset != nullptr)
 		{
-			IMediaPlayerPtr MediaPlayer = MediaAsset->GetMediaPlayer();
-
-			if (MediaPlayer.IsValid())
-			{
-				MediaPlayer->OnClosing().AddUObject(this, &UMediaTexture::HandleMediaPlayerClosing);
-				MediaPlayer->OnOpened().AddUObject(this, &UMediaTexture::HandleMediaPlayerOpened);
-			}
+			MediaAsset->OnMediaChanged().AddUObject(this, &UMediaTexture::HandleMediaAssetMediaChanged);
 		}	
 	}
 
@@ -195,10 +185,9 @@ void UMediaTexture::InitializeTrack( )
 	if (VideoTrack.IsValid())
 	{
 		VideoTrack->RemoveSink(VideoBuffer);
+		VideoTrack.Reset();
 	}
 
-	VideoTrack.Reset();
-
 	// initialize from new track
 	if (MediaAsset != nullptr)
 	{
@@ -232,16 +221,7 @@ void UMediaTexture::InitializeTrack( )
 /* UMediaTexture callbacks
  *****************************************************************************/
 
-void UMediaTexture::HandleMediaPlayerClosing( FString ClosingUrl )
+void UMediaTexture::HandleMediaAssetMediaChanged( )
 {
 	InitializeTrack();
 }
-
-
-void UMediaTexture::HandleMediaPlayerOpened( FString OpenedUrl )
-{
-	InitializeTrack();
-}
-
-
-#endif // WITH_EDITOR
diff --git a/Engine/Source/Runtime/MediaAssets/Public/MediaAsset.h b/Engine/Source/Runtime/MediaAssets/Public/MediaAsset.h
--- a/Engine/Source/Runtime/MediaAssets/Public/MediaAsset.h
+++ b/Engine/Source/Runtime/MediaAssets/Public/MediaAsset.h
@@ -21,8 +21,8 @@ enum EMediaAssetStreamModes
 };
 
 
-/** Multicast delegate that is invoked when a media asset's media is being closed. */
-DECLARE_DYNAMIC_MULTICAST_DELEGATE_OneParam(FOnMediaAssetMediaClosing, FString, ClosedUrl);
+/** Multicast delegate that is invoked when a media asset's media has been closed. */
+DECLARE_DYNAMIC_MULTICAST_DELEGATE(FOnMediaAssetMediaClosed);
 
 /** Multicast delegate that is invoked when a media asset's media has been opened. */
 DECLARE_DYNAMIC_MULTICAST_DELEGATE_OneParam(FOnMediaAssetMediaOpened, FString, OpenedUrl);
@@ -237,9 +237,16 @@ public:
 
 public:
 
-	/** Holds a delegate that is invoked when a media asset's media is being unloaded. */
+	/** Gets an event delegate that is invoked when media has been opened or closed. */
+	DECLARE_EVENT(UMediaAsset, FOnMediaChanged)
+	FOnMediaChanged& OnMediaChanged( )
+	{
+		return MediaChangedEvent;
+	}
+
+	/** Holds a delegate that is invoked when a media asset's media has been closed. */
 	UPROPERTY(BlueprintAssignable, Category="Media|MediaAsset")
-	FOnMediaAssetMediaClosing OnMediaClosing;
+	FOnMediaAssetMediaClosed OnMediaClosed;
 
 	/** Holds a delegate that is invoked when a media asset's media has been opened. */
 	UPROPERTY(BlueprintAssignable, Category="Media|MediaAsset")
@@ -298,10 +305,10 @@ protected:
 
 private:
 
-	// Callback for when the media player is unloading media.
-	void HandleMediaPlayerMediaClosing( FString ClosedUrl );
+	/** Callback for when the media player has closed the media. */
+	void HandleMediaPlayerMediaClosed( );
 
-	// Callback for when the media player is loading media.
+	/** Callback for when the media player has opened new media. */
 	void HandleMediaPlayerMediaOpened( FString OpenedUrl );
 
 private:
@@ -311,4 +318,9 @@ private:
 
 	/** Holds the media player used to play the media file. */
 	TSharedPtr<class IMediaPlayer> MediaPlayer;
+
+private:
+
+	/** Holds a delegate that is executed when media has been opened or closed. */
+	FOnMediaChanged MediaChangedEvent;
 };
diff --git a/Engine/Source/Runtime/MediaAssets/Public/MediaSoundWave.h b/Engine/Source/Runtime/MediaAssets/Public/MediaSoundWave.h
--- a/Engine/Source/Runtime/MediaAssets/Public/MediaSoundWave.h
+++ b/Engine/Source/Runtime/MediaAssets/Public/MediaSoundWave.h
@@ -70,16 +70,13 @@ public:
 
 protected:
 
-	/** Initializes the audio stream. */
-	void InitializeStream( );
+	/** Initializes the audio track. */
+	void InitializeTrack( );
 
 private:
 
-	// Callback for when the media player unloads media.
-	void HandleMediaPlayerClosing( FString ClosingUrl );
-
-	// Callback for when the media player unloads media.
-	void HandleMediaPlayerOpened( FString OpenedUrl );
+	/** Callback for when the media asset changed its media. */
+	void HandleMediaAssetMediaChanged( );
 
 private:
 
diff --git a/Engine/Source/Runtime/MediaAssets/Public/MediaTexture.h b/Engine/Source/Runtime/MediaAssets/Public/MediaTexture.h
--- a/Engine/Source/Runtime/MediaAssets/Public/MediaTexture.h
+++ b/Engine/Source/Runtime/MediaAssets/Public/MediaTexture.h
@@ -5,6 +5,7 @@
 #include "MediaSampleBuffer.h"
 #include "MediaTexture.generated.h"
 
+
 class UMediaAsset;
 
 
@@ -125,11 +126,8 @@ protected:
 
 private:
 
-	// Callback for when the media player unloads media.
-	void HandleMediaPlayerClosing( FString ClosingUrl );
-
-	// Callback for when the media player unloads media.
-	void HandleMediaPlayerOpened( FString OpenedUrl );
+	/** Callback for when the media asset changed its media. */
+	void HandleMediaAssetMediaChanged( );
 
 private:
 
