Media: Changed default clear color for media textures to black

diff --git a/Engine/Source/Runtime/MediaAssets/Private/Assets/MediaTexture.cpp b/Engine/Source/Runtime/MediaAssets/Private/Assets/MediaTexture.cpp
--- a/Engine/Source/Runtime/MediaAssets/Private/Assets/MediaTexture.cpp
+++ b/Engine/Source/Runtime/MediaAssets/Private/Assets/MediaTexture.cpp
@@ -9,7 +9,7 @@
 
 UMediaTexture::UMediaTexture( const class FPostConstructInitializeProperties& PCIP )
 	: Super(PCIP)
-	, ClearColor(FLinearColor::Red)
+	, ClearColor(FLinearColor::Black)
 	, MediaPlayer(nullptr)
 	, CurrentMediaPlayer(nullptr)
 	, VideoBuffer(MakeShareable(new FMediaSampleBuffer))
