Media: Fixed media sources not being found in staged builds

#UpgradeNotes: You will have to reassign the URL field in your MediaPlayer assets

diff --git a/Engine/Source/Developer/LauncherServices/Private/Launcher/LauncherWorker.cpp b/Engine/Source/Developer/LauncherServices/Private/Launcher/LauncherWorker.cpp
--- a/Engine/Source/Developer/LauncherServices/Private/Launcher/LauncherWorker.cpp
+++ b/Engine/Source/Developer/LauncherServices/Private/Launcher/LauncherWorker.cpp
@@ -225,7 +225,7 @@ FString FLauncherWorker::CreateUATCommand( const ILauncherProfileRef& InProfile,
 		}
 		else if (PlatformInfo->TargetPlatformName == FName("WindowsNoEditor") || PlatformInfo->TargetPlatformName == FName("Windows"))
 		{
-			Platforms += TEXT("+Win64");
+			Platforms += TEXT("+Win32");
 		}
 		else if (PlatformInfo->TargetPlatformName == FName("MacNoEditor"))
 		{
diff --git a/Engine/Source/Editor/DetailCustomizations/Private/MediaPlayerCustomization.cpp b/Engine/Source/Editor/DetailCustomizations/Private/MediaPlayerCustomization.cpp
--- a/Engine/Source/Editor/DetailCustomizations/Private/MediaPlayerCustomization.cpp
+++ b/Engine/Source/Editor/DetailCustomizations/Private/MediaPlayerCustomization.cpp
@@ -396,9 +396,12 @@ void FMediaPlayerCustomization::HandleUrlPickerPathPicked( const FString& Picked
 		FPaths::NormalizeFilename(FullPath);
 	}
 
-	if (FullPath.StartsWith(FPaths::RootDir()))
+	FString FullGameContentDir = FPaths::ConvertRelativePathToFull(FPaths::GameContentDir());
+
+	if (FullPath.StartsWith(FullGameContentDir))
 	{
-		FPaths::MakePathRelativeTo(FullPath, FPlatformProcess::BaseDir());
+		FPaths::MakePathRelativeTo(FullPath, *FullGameContentDir);
+		FullPath = FString(TEXT("./")) + FullPath;
 	}
 
 	UrlProperty->SetValue(FullPath);
@@ -414,12 +417,12 @@ EVisibility FMediaPlayerCustomization::HandleUrlWarningIconVisibility( ) const
 		return EVisibility::Hidden;
 	}
 
-	const FString FullPath = FPaths::ConvertRelativePathToFull(Url);
 	const FString FullMoviesPath = FPaths::ConvertRelativePathToFull(FPaths::GameContentDir() / TEXT("Movies"));
+	const FString FullUrl = FPaths::ConvertRelativePathToFull(FPaths::IsRelative(Url) ? FPaths::GameContentDir() / Url : Url);
 
-	if (FullPath.StartsWith(FullMoviesPath))
+	if (FullUrl.StartsWith(FullMoviesPath))
 	{
-		if (FPaths::FileExists(FullPath))
+		if (FPaths::FileExists(FullUrl))
 		{
 			return EVisibility::Hidden;
 		}
diff --git a/Engine/Source/Runtime/MediaAssets/Private/Assets/MediaPlayer.cpp b/Engine/Source/Runtime/MediaAssets/Private/Assets/MediaPlayer.cpp
--- a/Engine/Source/Runtime/MediaAssets/Private/Assets/MediaPlayer.cpp
+++ b/Engine/Source/Runtime/MediaAssets/Private/Assets/MediaPlayer.cpp
@@ -241,15 +241,16 @@ void UMediaPlayer::InitializePlayer( )
 		Player->OnOpened().AddUObject(this, &UMediaPlayer::HandleMediaPlayerMediaOpened);
 
 		// open the new media file
+		const FString FullUrl = FPaths::ConvertRelativePathToFull(FPaths::IsRelative(URL) ? FPaths::GameContentDir() / URL : URL);
 		bool OpenedSuccessfully = false;
 
 		if (StreamMode == EMediaPlayerStreamModes::MASM_FromUrl)
 		{
-			OpenedSuccessfully = Player->Open(FPaths::ConvertRelativePathToFull(URL));
+			OpenedSuccessfully = Player->Open(FPaths::ConvertRelativePathToFull(FullUrl));
 		}
-		else if (FPaths::FileExists(URL))
+		else if (FPaths::FileExists(FullUrl))
 		{
-			FArchive* FileReader = IFileManager::Get().CreateFileReader(*URL);
+			FArchive* FileReader = IFileManager::Get().CreateFileReader(*FullUrl);
 		
 			if (FileReader == nullptr)
 			{
@@ -263,7 +264,7 @@ void UMediaPlayer::InitializePlayer( )
 				FileData->AddUninitialized(FileReader->TotalSize());
 				FileReader->Serialize(FileData->GetData(), FileReader->TotalSize());
 
-				OpenedSuccessfully = Player->Open(MakeShareable(FileData), URL);
+				OpenedSuccessfully = Player->Open(MakeShareable(FileData), FullUrl);
 			}
 
 			delete FileReader;
@@ -272,7 +273,7 @@ void UMediaPlayer::InitializePlayer( )
 		// finish initialization
 		if (OpenedSuccessfully)
 		{
-			CurrentUrl = URL;
+			CurrentUrl = FullUrl;
 		}
 	}
 
