Media: Include what you use

diff --git a/Engine/Source/Runtime/MediaAssets/Private/Assets/MediaPlayer.cpp b/Engine/Source/Runtime/MediaAssets/Private/Assets/MediaPlayer.cpp
--- a/Engine/Source/Runtime/MediaAssets/Private/Assets/MediaPlayer.cpp
+++ b/Engine/Source/Runtime/MediaAssets/Private/Assets/MediaPlayer.cpp
@@ -1,6 +1,8 @@
 // Copyright 1998-2014 Epic Games, Inc. All Rights Reserved.
 
 #include "MediaAssetsPrivatePCH.h"
+#include "IMediaModule.h"
+#include "ModuleManager.h"
 
 
 /* UMediaPlayer structors
@@ -230,7 +232,14 @@ void UMediaPlayer::InitializePlayer()
 		}
 
 		// create new player
-		Player = IMediaModule::Get().CreatePlayer(URL);
+		IMediaModule* MediaModule = FModuleManager::LoadModulePtr<IMediaModule>("Media");
+
+		if (MediaModule == nullptr)
+		{
+			return;
+		}
+
+		Player = MediaModule->CreatePlayer(URL);
 
 		if (!Player.IsValid())
 		{
diff --git a/Engine/Source/Runtime/MediaAssets/Private/Assets/MediaSoundWave.cpp b/Engine/Source/Runtime/MediaAssets/Private/Assets/MediaSoundWave.cpp
--- a/Engine/Source/Runtime/MediaAssets/Private/Assets/MediaSoundWave.cpp
+++ b/Engine/Source/Runtime/MediaAssets/Private/Assets/MediaSoundWave.cpp
@@ -1,6 +1,9 @@
 // Copyright 1998-2014 Epic Games, Inc. All Rights Reserved.
 
 #include "MediaAssetsPrivatePCH.h"
+#include "IMediaTrackAudioDetails.h"
+#include "MediaSampleQueue.h"
+#include "MediaSoundWave.h"
 
 
 /* UMediaSoundWave structors
diff --git a/Engine/Source/Runtime/MediaAssets/Private/Assets/MediaTexture.cpp b/Engine/Source/Runtime/MediaAssets/Private/Assets/MediaTexture.cpp
--- a/Engine/Source/Runtime/MediaAssets/Private/Assets/MediaTexture.cpp
+++ b/Engine/Source/Runtime/MediaAssets/Private/Assets/MediaTexture.cpp
@@ -1,7 +1,8 @@
 // Copyright 1998-2014 Epic Games, Inc. All Rights Reserved.
 
 #include "MediaAssetsPrivatePCH.h"
-#include "MediaTextureResource.h"
+#include "IMediaTrackVideoDetails.h"
+#include "MediaSampleBuffer.h"
 
 
 /* UMediaTexture structors
diff --git a/Engine/Source/Runtime/MediaAssets/Private/Assets/MediaTextureResource.h b/Engine/Source/Runtime/MediaAssets/Private/Assets/MediaTextureResource.h
--- a/Engine/Source/Runtime/MediaAssets/Private/Assets/MediaTextureResource.h
+++ b/Engine/Source/Runtime/MediaAssets/Private/Assets/MediaTextureResource.h
@@ -3,6 +3,11 @@
 #pragma once
 
 
+// forward declarations
+class FMediaSampleBuffer;
+class UMediaTexture;
+
+
 /**
  * FTextureResource type for movie textures.
  */
@@ -19,7 +24,7 @@ public:
 	 * @param InOwner The Movie texture object to create a resource for (must not be nullptr).
 	 * @param InVideoBuffer The video sample buffer to use.
 	 */
-	FMediaTextureResource( const class UMediaTexture* InOwner, const FMediaSampleBufferRef& InVideoBuffer );
+	FMediaTextureResource(const class UMediaTexture* InOwner, const TSharedRef<FMediaSampleBuffer, ESPMode::ThreadSafe>& InVideoBuffer);
 
 public:
 
@@ -42,21 +47,21 @@ public:
 
 private:
 
-	// Whether the resource was cleared last frame.
+	/** Whether the resource was cleared last frame. */
 	bool Cleared;
 
-	// The color that the resource was cleared with.
+	/** The color that the resource was cleared with. */
 	FLinearColor LastClearColor;
 
-	// The playback time of the last drawn video frame.
+	/** The playback time of the last drawn video frame. */
 	FTimespan LastFrameTime;
 
 	/** The UTextureRenderTarget2D which this resource represents. */
-	const class UMediaTexture* Owner;
+	const UMediaTexture* Owner;
 
 	/** Texture resource used for rendering with and resolving to. */
 	FTexture2DRHIRef Texture2DRHI;
 
 	/** Pointer to the video sample buffer. */
-	FMediaSampleBufferRef VideoBuffer;
+	TSharedRef<FMediaSampleBuffer, ESPMode::ThreadSafe> VideoBuffer;
 };
diff --git a/Engine/Source/Runtime/MediaAssets/Private/MediaAssetsModule.cpp b/Engine/Source/Runtime/MediaAssets/Private/MediaAssetsModule.cpp
--- a/Engine/Source/Runtime/MediaAssets/Private/MediaAssetsModule.cpp
+++ b/Engine/Source/Runtime/MediaAssets/Private/MediaAssetsModule.cpp
@@ -1,7 +1,7 @@
 // Copyright 1998-2014 Epic Games, Inc. All Rights Reserved.
 
 #include "MediaAssetsPrivatePCH.h"
-#include "ModuleManager.h"
+#include "ModuleInterface.h"
 
 
 /**
diff --git a/Engine/Source/Runtime/MediaAssets/Private/MediaAssetsPrivatePCH.h b/Engine/Source/Runtime/MediaAssets/Private/MediaAssetsPrivatePCH.h
--- a/Engine/Source/Runtime/MediaAssets/Private/MediaAssetsPrivatePCH.h
+++ b/Engine/Source/Runtime/MediaAssets/Private/MediaAssetsPrivatePCH.h
@@ -9,13 +9,13 @@
 #include "Core.h"
 #include "CoreUObject.h"
 #include "Engine.h"
-#include "Media.h"
-#include "MediaPlayer.h"
-#include "MediaSoundWave.h"
-#include "MediaTexture.h"
+#include "IMediaPlayer.h"
+#include "MediaSampleBuffer.h"
 
 
 /* Private includes
  *****************************************************************************/
 
+#include "MediaPlayer.h"
+#include "MediaTexture.h"
 #include "MediaTextureResource.h"
diff --git a/Engine/Source/Runtime/MediaAssets/Public/MediaPlayer.h b/Engine/Source/Runtime/MediaAssets/Public/MediaPlayer.h
--- a/Engine/Source/Runtime/MediaAssets/Public/MediaPlayer.h
+++ b/Engine/Source/Runtime/MediaAssets/Public/MediaPlayer.h
@@ -5,6 +5,11 @@
 #include "MediaPlayer.generated.h"
 
 
+// forward declarations
+enum EMediaPlayerStreamModes;
+class IMediaPlayer;
+
+
 /**
  * Enumerates available media streaming modes.
  */
@@ -259,7 +264,7 @@ public:
 	 *
 	 * @return The player, or nullptr if no player was created.
 	 */
-	TSharedPtr<class IMediaPlayer> GetPlayer() const
+	TSharedPtr<IMediaPlayer> GetPlayer() const
 	{
 		return Player;
 	}
@@ -297,7 +302,7 @@ protected:
 
 	/** Select where to stream the media from, i.e. file or memory. */
 	UPROPERTY(EditAnywhere, BlueprintReadOnly, Category=Source)
-	TEnumAsByte<enum EMediaPlayerStreamModes> StreamMode;
+	TEnumAsByte<EMediaPlayerStreamModes> StreamMode;
 
 	/** The path or URL to the media file to be played. */
 	UPROPERTY(EditAnywhere, Category=Source)
@@ -317,7 +322,7 @@ private:
 	FString CurrentUrl;
 
 	/** Holds the low-level player used to play the media source. */
-	TSharedPtr<class IMediaPlayer> Player;
+	TSharedPtr<IMediaPlayer> Player;
 
 private:
 
diff --git a/Engine/Source/Runtime/MediaAssets/Public/MediaSoundWave.h b/Engine/Source/Runtime/MediaAssets/Public/MediaSoundWave.h
--- a/Engine/Source/Runtime/MediaAssets/Public/MediaSoundWave.h
+++ b/Engine/Source/Runtime/MediaAssets/Public/MediaSoundWave.h
@@ -6,6 +6,10 @@
 #include "Sound/SoundWave.h"
 #include "MediaSoundWave.generated.h"
 
+
+// forward declarations
+class FMediaSampleQueue;
+class IMediaPlayer;
 class UMediaPlayer;
 
 
@@ -48,7 +52,7 @@ public:
 	 *
 	 * @return The player, or nullptr if no player is available.
 	 */
-	TSharedPtr<class IMediaPlayer> GetPlayer() const;
+	TSharedPtr<IMediaPlayer> GetPlayer() const;
 
 public:
 
@@ -81,10 +85,10 @@ private:
 private:
 
 	/** The audio sample queue. */
-	FMediaSampleQueueRef AudioQueue;
+	TSharedRef<FMediaSampleQueue, ESPMode::ThreadSafe> AudioQueue;
 
 	/** Holds the selected audio track. */
-	IMediaTrackPtr AudioTrack;
+	TSharedPtr<IMediaTrack, ESPMode::ThreadSafe> AudioTrack;
 
 	/** Holds the media player asset currently being used. */
 	UMediaPlayer* CurrentMediaPlayer;
diff --git a/Engine/Source/Runtime/MediaAssets/Public/MediaTexture.h b/Engine/Source/Runtime/MediaAssets/Public/MediaTexture.h
--- a/Engine/Source/Runtime/MediaAssets/Public/MediaTexture.h
+++ b/Engine/Source/Runtime/MediaAssets/Public/MediaTexture.h
@@ -2,10 +2,15 @@
 
 #pragma once
 
-#include "MediaSampleBuffer.h"
+#include "Engine/Texture.h"
 #include "MediaTexture.generated.h"
 
 
+// forward declarations
+class IMediaPlayer;
+class IMediaTrack;
+enum EPixelFormat;
+enum TextureAddress;
 class UMediaPlayer;
 
 
@@ -20,11 +25,11 @@ class MEDIAASSETS_API UMediaTexture
 
 	/** The addressing mode to use for the X axis. */
 	UPROPERTY(EditAnywhere, BlueprintReadWrite, Category=MediaTexture, AssetRegistrySearchable)
-	TEnumAsByte<enum TextureAddress> AddressX;
+	TEnumAsByte<TextureAddress> AddressX;
 
 	/** The addressing mode to use for the Y axis. */
 	UPROPERTY(EditAnywhere, BlueprintReadWrite, Category=MediaTexture, AssetRegistrySearchable)
-	TEnumAsByte<enum TextureAddress> AddressY;
+	TEnumAsByte<TextureAddress> AddressY;
 
 	/** The color used to clear the texture if no video data is drawn. */
 	UPROPERTY(EditAnywhere, BlueprintReadWrite, Category=MediaTexture)
@@ -56,7 +61,7 @@ public:
 	 *
 	 * @return Pixel format (always PF_B8G8R8A8 for all movie textures).
 	 */
-	TEnumAsByte<enum EPixelFormat> GetFormat() const
+	TEnumAsByte<EPixelFormat> GetFormat() const
 	{
 		return PF_B8G8R8A8;
 	}
@@ -66,14 +71,14 @@ public:
 	 *
 	 * @return The player, or nullptr if no player is available.
 	 */
-	TSharedPtr<class IMediaPlayer> GetPlayer() const;
+	TSharedPtr<IMediaPlayer> GetPlayer() const;
 
 	/**
 	 * Gets the currently selected video track, if any.
 	 *
 	 * @return The selected video track, or nullptr if none is selected.
 	 */
-	TSharedPtr<class IMediaTrack, ESPMode::ThreadSafe> GetVideoTrack() const
+	TSharedPtr<IMediaTrack, ESPMode::ThreadSafe> GetVideoTrack() const
 	{
 		return VideoTrack;
 	}
@@ -141,8 +146,8 @@ private:
 	FRenderCommandFence* ReleasePlayerFence;
 
 	/** The video sample buffer. */
-	FMediaSampleBufferRef VideoBuffer;
+	TSharedRef<FMediaSampleBuffer, ESPMode::ThreadSafe> VideoBuffer;
 
 	/** Holds the selected video track. */
-	IMediaTrackPtr VideoTrack;
+	TSharedPtr<IMediaTrack, ESPMode::ThreadSafe> VideoTrack;
 };
