- Removed forward declaration of enum that breaks clang (and is not needed when it's declared 5 lines down)

diff --git a/Engine/Source/Runtime/MediaAssets/Public/MediaPlayer.h b/Engine/Source/Runtime/MediaAssets/Public/MediaPlayer.h
--- a/Engine/Source/Runtime/MediaAssets/Public/MediaPlayer.h
+++ b/Engine/Source/Runtime/MediaAssets/Public/MediaPlayer.h
@@ -6,7 +6,6 @@
 
 
 // forward declarations
-enum EMediaPlayerStreamModes;
 class IMediaPlayer;
 
 
