Media: Fixed media reloading when changing a MediaPlayer property in the Editor

diff --git a/Engine/Source/Runtime/MediaAssets/Private/Assets/MediaPlayer.cpp b/Engine/Source/Runtime/MediaAssets/Private/Assets/MediaPlayer.cpp
--- a/Engine/Source/Runtime/MediaAssets/Private/Assets/MediaPlayer.cpp
+++ b/Engine/Source/Runtime/MediaAssets/Private/Assets/MediaPlayer.cpp
@@ -282,7 +282,7 @@ void UMediaPlayer::InitializePlayer()
 		// finish initialization
 		if (OpenedSuccessfully)
 		{
-			CurrentUrl = FullUrl;
+			CurrentUrl = URL;
 		}
 	}
 
