MediaAssets: Fixed missing forward declarations for Media types

diff --git a/Engine/Source/Runtime/MediaAssets/Public/MediaPlayer.h b/Engine/Source/Runtime/MediaAssets/Public/MediaPlayer.h
--- a/Engine/Source/Runtime/MediaAssets/Public/MediaPlayer.h
+++ b/Engine/Source/Runtime/MediaAssets/Public/MediaPlayer.h
@@ -5,7 +5,6 @@
 #include "MediaPlayer.generated.h"
 
 
-// forward declarations
 class IMediaPlayer;
 
 
diff --git a/Engine/Source/Runtime/MediaAssets/Public/MediaSoundWave.h b/Engine/Source/Runtime/MediaAssets/Public/MediaSoundWave.h
--- a/Engine/Source/Runtime/MediaAssets/Public/MediaSoundWave.h
+++ b/Engine/Source/Runtime/MediaAssets/Public/MediaSoundWave.h
@@ -7,9 +7,9 @@
 #include "MediaSoundWave.generated.h"
 
 
-// forward declarations
 class FMediaSampleQueue;
 class IMediaPlayer;
+class IMediaTrack;
 class UMediaPlayer;
 
 
diff --git a/Engine/Source/Runtime/MediaAssets/Public/MediaTexture.h b/Engine/Source/Runtime/MediaAssets/Public/MediaTexture.h
--- a/Engine/Source/Runtime/MediaAssets/Public/MediaTexture.h
+++ b/Engine/Source/Runtime/MediaAssets/Public/MediaTexture.h
@@ -6,12 +6,12 @@
 #include "MediaTexture.generated.h"
 
 
-// forward declarations
+class FMediaSampleBuffer;
 class IMediaPlayer;
 class IMediaTrack;
+class UMediaPlayer;
 enum EPixelFormat;
 enum TextureAddress;
-class UMediaPlayer;
 
 
 /**
@@ -127,8 +127,6 @@ protected:
 	/** Initializes the video track. */
 	void InitializeTrack();
 
-	void PickDefaultTrack();
-
 private:
 
 	/** Callback for when the UMediaPlayer asset changed its media. */
