Media: Rendering ClearColor while media is not playing/not paused

This prevents the last frame of video from being visible after a non-looping movie finished playing.

diff --git a/Engine/Source/Runtime/MediaAssets/Private/Assets/MediaTextureResource.cpp b/Engine/Source/Runtime/MediaAssets/Private/Assets/MediaTextureResource.cpp
--- a/Engine/Source/Runtime/MediaAssets/Private/Assets/MediaTextureResource.cpp
+++ b/Engine/Source/Runtime/MediaAssets/Private/Assets/MediaTextureResource.cpp
@@ -86,11 +86,18 @@ FIntPoint FMediaTextureResource::GetSizeXY() const
 
 void FMediaTextureResource::UpdateDeferredResource(FRHICommandListImmediate& RHICmdList, bool bClearRenderTarget/*=true*/)
 {
-	FTimespan CurrentFrameTime = VideoBuffer->GetCurrentSampleTime();
-	TSharedPtr<TArray<uint8>, ESPMode::ThreadSafe> CurrentFrame = VideoBuffer->GetCurrentSample();
+	TSharedPtr<TArray<uint8>, ESPMode::ThreadSafe> CurrentFrame;
+	TSharedPtr<IMediaPlayer> Player = Owner->GetPlayer();
+
+	if (Player.IsValid() && (Player->IsPlaying() || Player->IsPaused()))
+	{
+		CurrentFrame = VideoBuffer->GetCurrentSample();
+	}
 
 	if (CurrentFrame.IsValid())
 	{
+		FTimespan CurrentFrameTime = VideoBuffer->GetCurrentSampleTime();
+
 		// draw the latest video frame
 		if (CurrentFrameTime != LastFrameTime)
 		{
