First pass at UMF support on PS4.

diff --git a/Engine/Source/Runtime/MediaAssets/Private/Assets/MediaSoundWave.cpp b/Engine/Source/Runtime/MediaAssets/Private/Assets/MediaSoundWave.cpp
--- a/Engine/Source/Runtime/MediaAssets/Private/Assets/MediaSoundWave.cpp
+++ b/Engine/Source/Runtime/MediaAssets/Private/Assets/MediaSoundWave.cpp
@@ -126,6 +126,16 @@ void UMediaSoundWave::Serialize( FArchive& Ar )
 	USoundBase::Serialize(Ar);
 }
 
+void UMediaSoundWave::PostLoad()
+{
+	Super::PostLoad();
+
+	if (!HasAnyFlags(RF_ClassDefaultObject) && !GIsBuildMachine)
+	{
+		InitializeTrack();
+	}
+}
+
 
 /* UMediaSoundWave implementation
  *****************************************************************************/
diff --git a/Engine/Source/Runtime/MediaAssets/Private/Assets/MediaTextureResource.cpp b/Engine/Source/Runtime/MediaAssets/Private/Assets/MediaTextureResource.cpp
--- a/Engine/Source/Runtime/MediaAssets/Private/Assets/MediaTextureResource.cpp
+++ b/Engine/Source/Runtime/MediaAssets/Private/Assets/MediaTextureResource.cpp
@@ -23,6 +23,7 @@ void FMediaTextureResource::InitDynamicRHI()
 	{
 		// Create the RHI texture. Only one mip is used and the texture is targetable or resolve.
 		uint32 TexCreateFlags = Owner->SRGB ? TexCreate_SRGB : 0;
+		TexCreateFlags |= TexCreate_Dynamic | TexCreate_NoTiling;
 		FRHIResourceCreateInfo CreateInfo;
 
 		RHICreateTargetableShaderResource2D(
diff --git a/Engine/Source/Runtime/MediaAssets/Public/MediaSoundWave.h b/Engine/Source/Runtime/MediaAssets/Public/MediaSoundWave.h
--- a/Engine/Source/Runtime/MediaAssets/Public/MediaSoundWave.h
+++ b/Engine/Source/Runtime/MediaAssets/Public/MediaSoundWave.h
@@ -71,6 +71,7 @@ public:
 	virtual void GetAssetRegistryTags( TArray<FAssetRegistryTag>& OutTags ) const override;
 	virtual SIZE_T GetResourceSize( EResourceSizeMode::Type Mode ) override;
 	virtual void Serialize( FArchive& Ar ) override;
+	virtual void PostLoad() override;
 
 protected:
 
