Media: Removed AutoPlay feature in MediaPlayer assets.

#UpgradeNotes:
The AutoPlay feature turned out to be quite dangerous, because it causes media playback to start as soon as a MediaPlayer is loaded, thus consuming resources even if it is not used anywhere. Media should be played and stopped on demand using Blueprints (MediaPlayer.Play) or C++ instead.

diff --git a/Engine/Source/Runtime/MediaAssets/Private/Assets/MediaPlayer.cpp b/Engine/Source/Runtime/MediaAssets/Private/Assets/MediaPlayer.cpp
--- a/Engine/Source/Runtime/MediaAssets/Private/Assets/MediaPlayer.cpp
+++ b/Engine/Source/Runtime/MediaAssets/Private/Assets/MediaPlayer.cpp
@@ -10,8 +10,6 @@
 
 UMediaPlayer::UMediaPlayer( const FObjectInitializer& ObjectInitializer )
 	: Super(ObjectInitializer)
-	, AutoPlay(false)
-	, AutoPlayRate(1.0f)
 	, Looping(true)
 	, StreamMode(MASM_FromUrl)
 	, Player(nullptr)
@@ -286,17 +284,9 @@ void UMediaPlayer::InitializePlayer()
 		}
 	}
 
-	if (!Player.IsValid())
-	{
-		return;
-	}
-
-	// start playback, if desired
-	Player->SetLooping(Looping);
-    
-	if (AutoPlay)
+	if (Player.IsValid())
 	{
-		Player->SetRate(AutoPlayRate);
+		Player->SetLooping(Looping);
 	}
 }
 
diff --git a/Engine/Source/Runtime/MediaAssets/Public/MediaPlayer.h b/Engine/Source/Runtime/MediaAssets/Public/MediaPlayer.h
--- a/Engine/Source/Runtime/MediaAssets/Public/MediaPlayer.h
+++ b/Engine/Source/Runtime/MediaAssets/Public/MediaPlayer.h
@@ -286,14 +286,6 @@ protected:
 
 protected:
 
-	/** Whether this media should automatically start playing when it is loaded. */
-	UPROPERTY(EditAnywhere, Category=Playback)
-	uint32 AutoPlay:1;
-
-	/** The default playback rate when auto play is enabled. */
-	UPROPERTY(EditAnywhere, Category=Playback)
-	float AutoPlayRate;
-
 	/** Whether playback should loop when it reaches the end. */
 	UPROPERTY(EditAnywhere, Category=Playback)
 	uint32 Looping:1;
