Media: Fixed MediaPlayer not handling web URLs and non-relative paths properly

diff --git a/Engine/Source/Runtime/MediaAssets/Private/Assets/MediaPlayer.cpp b/Engine/Source/Runtime/MediaAssets/Private/Assets/MediaPlayer.cpp
--- a/Engine/Source/Runtime/MediaAssets/Private/Assets/MediaPlayer.cpp
+++ b/Engine/Source/Runtime/MediaAssets/Private/Assets/MediaPlayer.cpp
@@ -248,33 +248,41 @@ void UMediaPlayer::InitializePlayer()
 		Player->OnOpened().AddUObject(this, &UMediaPlayer::HandleMediaPlayerMediaOpened);
 
 		// open the new media file
-		const FString FullUrl = FPaths::ConvertRelativePathToFull(FPaths::IsRelative(URL) ? FPaths::GameContentDir() / URL : URL);
 		bool OpenedSuccessfully = false;
 
-		if (StreamMode == EMediaPlayerStreamModes::MASM_FromUrl)
+		if (URL.Contains(TEXT("://")))
 		{
-			OpenedSuccessfully = Player->Open(FPaths::ConvertRelativePathToFull(FullUrl));
+			OpenedSuccessfully = Player->Open(URL);
 		}
-		else if (FPaths::FileExists(FullUrl))
+		else
 		{
-			FArchive* FileReader = IFileManager::Get().CreateFileReader(*FullUrl);
-		
-			if (FileReader == nullptr)
+			const FString FullUrl = FPaths::ConvertRelativePathToFull(FPaths::IsRelative(URL) ? FPaths::GameContentDir() / URL : URL);
+
+			if (StreamMode == EMediaPlayerStreamModes::MASM_FromUrl)
 			{
-				return;
+				OpenedSuccessfully = Player->Open(FullUrl);
 			}
-
-			if (FileReader->TotalSize() > 0)
+			else if (FPaths::FileExists(FullUrl))
 			{
-				TArray<uint8>* FileData = new TArray<uint8>();
+				FArchive* FileReader = IFileManager::Get().CreateFileReader(*FullUrl);
+		
+				if (FileReader == nullptr)
+				{
+					return;
+				}
 
-				FileData->AddUninitialized(FileReader->TotalSize());
-				FileReader->Serialize(FileData->GetData(), FileReader->TotalSize());
+				if (FileReader->TotalSize() > 0)
+				{
+					TArray<uint8>* FileData = new TArray<uint8>();
 
-				OpenedSuccessfully = Player->Open(MakeShareable(FileData), FullUrl);
-			}
+					FileData->AddUninitialized(FileReader->TotalSize());
+					FileReader->Serialize(FileData->GetData(), FileReader->TotalSize());
+
+					OpenedSuccessfully = Player->Open(MakeShareable(FileData), FullUrl);
+				}
 
-			delete FileReader;
+				delete FileReader;
+			}
 		}
 
 		// finish initialization
