Media: Exposed Close method on MediaPlayer class and blueprint

diff --git a/Engine/Source/Runtime/MediaAssets/Private/Assets/MediaPlayer.cpp b/Engine/Source/Runtime/MediaAssets/Private/Assets/MediaPlayer.cpp
--- a/Engine/Source/Runtime/MediaAssets/Private/Assets/MediaPlayer.cpp
+++ b/Engine/Source/Runtime/MediaAssets/Private/Assets/MediaPlayer.cpp
@@ -31,6 +31,15 @@ bool UMediaPlayer::CanPlay() const
 }
 
 
+void UMediaPlayer::Close()
+{
+	if (Player.IsValid())
+	{
+		Player->Close();
+	}
+}
+
+
 FTimespan UMediaPlayer::GetDuration() const
 {
 	if (Player.IsValid())
diff --git a/Engine/Source/Runtime/MediaAssets/Public/MediaPlayer.h b/Engine/Source/Runtime/MediaAssets/Public/MediaPlayer.h
--- a/Engine/Source/Runtime/MediaAssets/Public/MediaPlayer.h
+++ b/Engine/Source/Runtime/MediaAssets/Public/MediaPlayer.h
@@ -65,6 +65,14 @@ public:
 	UFUNCTION(BlueprintCallable, Category="Media|MediaPlayer")
 	bool CanPlay() const;
 
+	/**
+	 * Close the currently open media, if any.
+	 *
+	 * @see OnMediaClosed, Open, Pause, Play
+	 */
+	UFUNCTION(BlueprintCallable, Category="Media|MediaPlayer")
+	void Close();
+
 	/**
 	 * Gets the media's duration.
 	 *
@@ -142,7 +150,7 @@ public:
 	 *
 	 * @param NewUrl The URL to open.
 	 * @return true on success, false otherwise.
-	 * @see GetUrl
+	 * @see GetUrl, Close
 	 */
 	UFUNCTION(BlueprintCallable, Category="Media|MediaPlayer")
 	bool OpenUrl( const FString& NewUrl );
@@ -153,7 +161,7 @@ public:
 	 * This is the same as setting the playback rate to 0.0.
 	 *
 	 * @return true if playback is being paused, false otherwise.
-	 * @see CanPause, Play, Rewind, Seek, SetRate
+	 * @see CanPause, Close, Play, Rewind, Seek, SetRate
 	 */
 	UFUNCTION(BlueprintCallable, Category="Media|MediaPlayer")
 	bool Pause();
@@ -164,7 +172,7 @@ public:
 	 * This is the same as setting the playback rate to 1.0.
 	 *
 	 * @return true if playback is starting, false otherwise.
-	 * @see CanPlay, Pause, Rewind, Seek, SetRate
+	 * @see CanPlay, Close, Pause, Rewind, Seek, SetRate
 	 */
 	UFUNCTION(BlueprintCallable, Category="Media|MediaPlayer")
 	bool Play();
@@ -175,7 +183,7 @@ public:
 	 * This is the same as seeking to zero time.
 	 *
 	 * @return true if rewinding, false otherwise.
-	 * @see GetTime, Pause, Play, Seek
+	 * @see GetTime, Close, Pause, Play, Seek
 	 */
 	UFUNCTION(BlueprintCallable, Category="Media|MediaPlayer")
 	bool Rewind();
