Forgot to check in

diff --git a/Engine/Source/Runtime/MediaAssets/Private/Assets/MediaTexture.cpp b/Engine/Source/Runtime/MediaAssets/Private/Assets/MediaTexture.cpp
--- a/Engine/Source/Runtime/MediaAssets/Private/Assets/MediaTexture.cpp
+++ b/Engine/Source/Runtime/MediaAssets/Private/Assets/MediaTexture.cpp
@@ -85,7 +85,7 @@ void UMediaTexture::UpdateResource()
 #if WITH_ENGINE
 	if (VideoTrack.IsValid() && Resource)
 	{		
-		VideoTrack->RemoveBoundTexture(Resource->TextureRHI.GetReference());
+		VideoTrack->UnbindTexture(Resource->TextureRHI.GetReference());
 	}
 #endif
 	UTexture::UpdateResource();
@@ -204,7 +204,7 @@ void UMediaTexture::InitializeTrack()
 		VideoTrack->GetStream().RemoveSink(VideoBuffer);
 
 #if WITH_ENGINE
-		VideoTrack->RemoveBoundTexture(Resource->TextureRHI.GetReference());
+		VideoTrack->UnbindTexture(Resource->TextureRHI.GetReference());
 #endif
 		VideoTrack.Reset();
 	}
@@ -249,7 +249,7 @@ void UMediaTexture::InitializeTrack()
 #if WITH_ENGINE
 		FlushRenderingCommands();
 		IMediaVideoTrack* VideoTrackPtr = static_cast<IMediaVideoTrack*>(VideoTrack.Get());
-		VideoTrackPtr->AddBoundTexture((Resource->TextureRHI.GetReference()));
+		VideoTrackPtr->BindTexture((Resource->TextureRHI.GetReference()));
 #endif
 	}
 }
