Fix MediaFramework GC problems.
Also fix a leak with PS4 moviestreamer.

diff --git a/Engine/Source/Runtime/MediaAssets/Private/Assets/MediaSoundWave.cpp b/Engine/Source/Runtime/MediaAssets/Private/Assets/MediaSoundWave.cpp
--- a/Engine/Source/Runtime/MediaAssets/Private/Assets/MediaSoundWave.cpp
+++ b/Engine/Source/Runtime/MediaAssets/Private/Assets/MediaSoundWave.cpp
@@ -14,26 +14,12 @@ UMediaSoundWave::UMediaSoundWave( const FObjectInitializer& ObjectInitializer )
 	, AudioTrackIndex(INDEX_NONE)
 	, AudioQueue(MakeShareable(new FMediaSampleQueue))
 {
+	bSetupDelegates = false;
 	bLooping = false;
 	bProcedural = true;
 	Duration = INDEFINITELY_LOOPING_DURATION;
 }
 
-
-UMediaSoundWave::~UMediaSoundWave()
-{
-	if (AudioTrack.IsValid())
-	{
-		AudioTrack->GetStream().RemoveSink(AudioQueue);
-	}
-
-	if (CurrentMediaPlayer != nullptr)
-	{
-		CurrentMediaPlayer->OnTracksChanged().RemoveAll(this);
-	}
-}
-
-
 /* UMediaSoundWave interface
  *****************************************************************************/
 
@@ -142,6 +128,23 @@ void UMediaSoundWave::PostLoad()
 	}
 }
 
+void UMediaSoundWave::BeginDestroy()
+{
+	Super::BeginDestroy();
+
+	if (AudioTrack.IsValid())
+	{
+		AudioTrack->GetStream().RemoveSink(AudioQueue);
+		AudioTrack.Reset();
+	}
+
+	if (CurrentMediaPlayer.IsValid())
+	{
+		CurrentMediaPlayer->OnTracksChanged().RemoveAll(this);
+		CurrentMediaPlayer.Reset();
+	}
+}
+
 
 /* UMediaSoundWave implementation
  *****************************************************************************/
@@ -149,7 +152,7 @@ void UMediaSoundWave::PostLoad()
 void UMediaSoundWave::InitializeTrack()
 {
 	// assign new media player asset
-	if (CurrentMediaPlayer != MediaPlayer)
+	if (CurrentMediaPlayer != MediaPlayer || !bSetupDelegates)
 	{
 		if (CurrentMediaPlayer != nullptr)
 		{
@@ -161,7 +164,8 @@ void UMediaSoundWave::InitializeTrack()
 		if (MediaPlayer != nullptr)
 		{
 			MediaPlayer->OnTracksChanged().AddUObject(this, &UMediaSoundWave::HandleMediaPlayerTracksChanged);
-		}	
+		}
+		bSetupDelegates = true;
 	}
 
 	// disconnect from current track
diff --git a/Engine/Source/Runtime/MediaAssets/Private/Assets/MediaTexture.cpp b/Engine/Source/Runtime/MediaAssets/Private/Assets/MediaTexture.cpp
--- a/Engine/Source/Runtime/MediaAssets/Private/Assets/MediaTexture.cpp
+++ b/Engine/Source/Runtime/MediaAssets/Private/Assets/MediaTexture.cpp
@@ -16,25 +16,12 @@ UMediaTexture::UMediaTexture( const FObjectInitializer& ObjectInitializer )
 	, CurrentMediaPlayer(nullptr)
 	, VideoBuffer(MakeShareable(new FMediaSampleBuffer))
 {
+	bDelegatesAdded = false;
 	NeverStream = true;
 	UpdateResource();
 }
 
 
-UMediaTexture::~UMediaTexture()
-{
-	if (VideoTrack.IsValid())
-	{
-		VideoTrack->GetStream().RemoveSink(VideoBuffer);
-	}
-
-	if (CurrentMediaPlayer != nullptr)
-	{
-		CurrentMediaPlayer->OnTracksChanged().RemoveAll(this);
-	}
-}
-
-
 /* UMediaTexture interface
  *****************************************************************************/
 
@@ -114,6 +101,18 @@ void UMediaTexture::FinishDestroy()
 	delete ReleasePlayerFence;
 	ReleasePlayerFence = nullptr;
 
+	if (VideoTrack.IsValid())
+	{
+		VideoTrack->GetStream().RemoveSink(VideoBuffer);
+		VideoTrack.Reset();
+	}
+
+	if (CurrentMediaPlayer.IsValid())
+	{
+		CurrentMediaPlayer->OnTracksChanged().RemoveAll(this);
+		CurrentMediaPlayer.Reset();
+	}
+
 	Super::FinishDestroy();
 }
 
@@ -183,7 +182,7 @@ void UMediaTexture::PostEditChangeProperty( FPropertyChangedEvent& PropertyChang
 void UMediaTexture::InitializeTrack()
 {
 	// assign new media player asset
-	if (CurrentMediaPlayer != MediaPlayer)
+	if (CurrentMediaPlayer != MediaPlayer || !bDelegatesAdded)
 	{
 		if (CurrentMediaPlayer != nullptr)
 		{
@@ -196,6 +195,7 @@ void UMediaTexture::InitializeTrack()
 		{
 			MediaPlayer->OnTracksChanged().AddUObject(this, &UMediaTexture::HandleMediaPlayerTracksChanged);
 		}	
+		bDelegatesAdded = true;
 	}
 
 	// disconnect from current track
diff --git a/Engine/Source/Runtime/MediaAssets/Public/MediaSoundWave.h b/Engine/Source/Runtime/MediaAssets/Public/MediaSoundWave.h
--- a/Engine/Source/Runtime/MediaAssets/Public/MediaSoundWave.h
+++ b/Engine/Source/Runtime/MediaAssets/Public/MediaSoundWave.h
@@ -30,11 +30,6 @@ class MEDIAASSETS_API UMediaSoundWave
 	UPROPERTY(EditAnywhere, BlueprintReadWrite, Category=MediaPlayer)
 	UMediaPlayer* MediaPlayer;
 
-public:
-
-	/** Destructor. */
-	~UMediaSoundWave();
-
 public:
 
 	/**
@@ -72,6 +67,7 @@ public:
 	virtual SIZE_T GetResourceSize( EResourceSizeMode::Type Mode ) override;
 	virtual void Serialize( FArchive& Ar ) override;
 	virtual void PostLoad() override;
+	virtual void BeginDestroy() override;
 
 protected:
 
@@ -92,8 +88,11 @@ private:
 	TSharedPtr<IMediaAudioTrack, ESPMode::ThreadSafe> AudioTrack;
 
 	/** Holds the media player asset currently being used. */
-	UMediaPlayer* CurrentMediaPlayer;
+	UPROPERTY()
+	TWeakObjectPtr<UMediaPlayer> CurrentMediaPlayer;
 
 	/** Holds queued audio samples. */
 	TArray<uint8> QueuedAudio;
+
+	bool bSetupDelegates;
 };
diff --git a/Engine/Source/Runtime/MediaAssets/Public/MediaTexture.h b/Engine/Source/Runtime/MediaAssets/Public/MediaTexture.h
--- a/Engine/Source/Runtime/MediaAssets/Public/MediaTexture.h
+++ b/Engine/Source/Runtime/MediaAssets/Public/MediaTexture.h
@@ -39,11 +39,6 @@ class MEDIAASSETS_API UMediaTexture
 	UPROPERTY(EditAnywhere, BlueprintReadWrite, Category=MediaPlayer)
 	int32 VideoTrackIndex;
 
-public:
-
-	/** Destructor. */
-	~UMediaTexture();
-
 public:
 
 	/**
@@ -139,7 +134,8 @@ private:
 	FIntPoint CachedDimensions;
 
 	/** Holds the UMediaPlayer asset currently being used. */
-	UMediaPlayer* CurrentMediaPlayer;
+	UPROPERTY() 
+	TWeakObjectPtr<UMediaPlayer> CurrentMediaPlayer;
 
 	/** Synchronizes access to this object from the render thread. */
 	FRenderCommandFence* ReleasePlayerFence;
@@ -149,4 +145,6 @@ private:
 
 	/** Holds the selected video track. */
 	TSharedPtr<IMediaVideoTrack, ESPMode::ThreadSafe> VideoTrack;
+
+	bool bDelegatesAdded;
 };
diff --git a/Engine/Source/Runtime/NetworkFile/Private/NetworkPlatformFile.cpp b/Engine/Source/Runtime/NetworkFile/Private/NetworkPlatformFile.cpp
--- a/Engine/Source/Runtime/NetworkFile/Private/NetworkPlatformFile.cpp
+++ b/Engine/Source/Runtime/NetworkFile/Private/NetworkPlatformFile.cpp
@@ -19,6 +19,7 @@
 
 DEFINE_LOG_CATEGORY(LogNetworkPlatformFile);
 
+FString FNetworkPlatformFile::MP4Extension = TEXT(".mp4");
 
 FNetworkPlatformFile::FNetworkPlatformFile()
 	: bHasLoadedDDCDirectories(false)
@@ -1021,6 +1022,18 @@ void AsyncReadUnsolicitedFiles(int32 InNumUnsolictedFiles, FNetworkPlatformFile&
 	(new FAutoDeleteAsyncTask<FAsyncReadUnsolicitedFile>(InNumUnsolictedFiles, &InNetworkFile, &InInnerPlatformFile, InServerEngineDir, InServerGameDir, InNetworkDoneEvent, InWritingDoneEvent))->StartSynchronousTask();
 }
 
+bool FNetworkPlatformFile::IsMediaExtension(const TCHAR* Ext)
+{
+	if (*Ext != TEXT('.'))
+	{
+		return MP4Extension.EndsWith(Ext);
+	}
+	else
+	{
+		return MP4Extension == Ext;
+	}
+}
+
 /**
  * Given a filename, make sure the file exists on the local filesystem
  */
@@ -1080,7 +1093,8 @@ void FNetworkPlatformFile::EnsureFileIsLocal(const FString& Filename)
 	//UE_LOG(LogNetworkPlatformFile, Display, TEXT("Check for local file %6.2fms - %s"), ThisTime, *Filename);
 
 	// this is a bit of a waste if we aren't doing cook on the fly, but we assume missing asset files are relatively rare
-	bool bIsCookable = GConfig && GConfig->IsReadyForUse() && FPackageName::IsPackageExtension(*FPaths::GetExtension(Filename, true));
+	FString Extension = FPaths::GetExtension(Filename, true);
+	bool bIsCookable = GConfig && GConfig->IsReadyForUse() && (FPackageName::IsPackageExtension(*Extension) || IsMediaExtension(*Extension));
 
 	// we only copy files that actually exist on the server, can greatly reduce network traffic for, say,
 	// the INT file each package tries to load
diff --git a/Engine/Source/Runtime/NetworkFile/Public/NetworkPlatformFile.h b/Engine/Source/Runtime/NetworkFile/Public/NetworkPlatformFile.h
--- a/Engine/Source/Runtime/NetworkFile/Public/NetworkPlatformFile.h
+++ b/Engine/Source/Runtime/NetworkFile/Public/NetworkPlatformFile.h
@@ -153,6 +153,15 @@ protected:
 
 private:
 
+	/**
+	* Returns whether the passed in extension is a video
+	* extension. Extensions with and without trailing dots are supported.
+	*
+	* @param	Extension to test.
+	* @return	True if Ext is a video extension.  e.g. .mp4
+	*/
+	static bool IsMediaExtension(const TCHAR* Ext);
+
 	/**
 	 * @return true if the path exists in a directory that should always use the local filesystem
 	 * This version does not worry about initialization or thread safety, do not call directly
@@ -222,6 +231,8 @@ private:
 
     // Our network Transport. 
 	class ITransport* Transport; 
+
+	static FString MP4Extension;
 };
 
 class SOCKETS_API FNetworkFileHandle : public IFileHandle
