Media: Not loading media on dedicated servers (UE-19832)

diff --git a/Engine/Source/Runtime/MediaAssets/Private/Assets/MediaPlayer.cpp b/Engine/Source/Runtime/MediaAssets/Private/Assets/MediaPlayer.cpp
--- a/Engine/Source/Runtime/MediaAssets/Private/Assets/MediaPlayer.cpp
+++ b/Engine/Source/Runtime/MediaAssets/Private/Assets/MediaPlayer.cpp
@@ -219,6 +219,11 @@ void UMediaPlayer::PostEditChangeProperty( FPropertyChangedEvent& PropertyChange
 
 void UMediaPlayer::InitializePlayer()
 {
+	if (IsRunningDedicatedServer())
+	{
+		return;
+	}
+
 	if (URL != CurrentUrl)
 	{
 		// close previous player
diff --git a/Engine/Source/Runtime/MediaAssets/Private/Assets/MediaTexture.cpp b/Engine/Source/Runtime/MediaAssets/Private/Assets/MediaTexture.cpp
--- a/Engine/Source/Runtime/MediaAssets/Private/Assets/MediaTexture.cpp
+++ b/Engine/Source/Runtime/MediaAssets/Private/Assets/MediaTexture.cpp
@@ -182,7 +182,7 @@ void UMediaTexture::PostEditChangeProperty( FPropertyChangedEvent& PropertyChang
 void UMediaTexture::InitializeTrack()
 {
 	// assign new media player asset
-	if (CurrentMediaPlayer != MediaPlayer || !bDelegatesAdded)
+	if ((CurrentMediaPlayer != MediaPlayer) || !bDelegatesAdded)
 	{
 		if (CurrentMediaPlayer != nullptr)
 		{
@@ -194,7 +194,8 @@ void UMediaTexture::InitializeTrack()
 		if (MediaPlayer != nullptr)
 		{
 			MediaPlayer->OnTracksChanged().AddUObject(this, &UMediaTexture::HandleMediaPlayerTracksChanged);
-		}	
+		}
+
 		bDelegatesAdded = true;
 	}
 
@@ -204,11 +205,12 @@ void UMediaTexture::InitializeTrack()
 		VideoTrack->GetStream().RemoveSink(VideoBuffer);
 
 #if WITH_ENGINE
-		if (Resource && Resource->TextureRHI.IsValid())
+		if ((Resource != nullptr) && Resource->TextureRHI.IsValid())
 		{
 			VideoTrack->UnbindTexture(Resource->TextureRHI.GetReference());
 		}
 #endif
+
 		VideoTrack.Reset();
 	}
 
@@ -252,7 +254,7 @@ void UMediaTexture::InitializeTrack()
 #if WITH_ENGINE
 		FlushRenderingCommands();
 
-		if (Resource && Resource->TextureRHI.IsValid())
+		if ((Resource != nullptr) && Resource->TextureRHI.IsValid())
 		{
 			IMediaVideoTrack* VideoTrackPtr = static_cast<IMediaVideoTrack*>(VideoTrack.Get());
 			VideoTrackPtr->BindTexture((Resource->TextureRHI.GetReference()));
