Copying //UE4/Dev-Sequencer to //UE4/Dev-Main (Source: //UE4/Dev-Sequencer @ 3057646)

diff --git a/Engine/Source/Runtime/Media/Public/IMediaTracks.h b/Engine/Source/Runtime/Media/Public/IMediaTracks.h
new file mode 100644
--- /dev/null
+++ b/Engine/Source/Runtime/Media/Public/IMediaTracks.h
@@ -0,0 +1,172 @@
+// Copyright 1998-2016 Epic Games, Inc. All Rights Reserved.
+
+#pragma once
+
+
+/**
+ * Enumerates available media track types.
+ */
+enum class EMediaTrackType
+{
+	/** Audio tracks. */
+	Audio,
+
+	/** Binary data tracks. */
+	Binary,
+
+	/** Caption (subtitle) tracks) */
+	Caption,
+
+	/** Still image tracks. */
+	Image,
+
+	/** Script tracks. */
+	Script,
+
+	/** Video tracks. */
+	Video
+};
+
+
+/**
+ * Interface for access to a media player's tracks.
+ *
+ * @see IMediaControls, IMediaOutput
+ */
+class IMediaTracks
+{
+public:
+
+	/**
+	 * Get the number of channels in the specified audio track.
+	 *
+	 * @param TrackIndex Index of the audio track.
+	 * @return Number of channels.
+	 * @see GetAudioTrackSamplesPerSeconds
+	 */
+	virtual uint32 GetAudioTrackChannels(int32 TrackIndex) const = 0;
+
+	/**
+	 * Get the sample rate of the specified audio track.
+	 *
+	 * @param TrackIndex Index of the audio track.
+	 * @return Samples per second.
+	 * @see GetAudioTrackChannels
+	 */
+	virtual uint32 GetAudioTrackSampleRate(int32 TrackIndex) const = 0;
+
+	/**
+	 * Get the number of media tracks of the given type.
+	 *
+	 * @param TrackType The type of media tracks.
+	 * @return Number of tracks.
+	 * @see GetSelectedTrack, SelectTrack
+	 */
+	virtual int32 GetNumTracks(EMediaTrackType TrackType) const = 0;
+
+	/**
+	 * Get the index of the currently selected track of the given type.
+	 *
+	 * @param TrackType The type of track to get.
+	 * @return The index of the selected track, or INDEX_NONE if no track is active.
+	 * @see GetNumTracks, SelectTrack
+	 */
+	virtual int32 GetSelectedTrack(EMediaTrackType TrackType) const = 0;
+
+	/**
+	 * Get the human readable name of the specified track.
+	 *
+	 * @param TrackType The type of track.
+	 * @param TrackIndex The index of the track.
+	 * @return Display name.
+	 * @see GetTrackLanguage, GetTrackName
+	 */
+	virtual FText GetTrackDisplayName(EMediaTrackType TrackType, int32 TrackIndex) const = 0;
+
+	/**
+	 * Get the language tag of the specified track.
+	 *
+	 * @param TrackType The type of track.
+	 * @param TrackIndex The index of the track.
+	 * @return Language tag, i.e. "en-US" for English, or "und" for undefined.
+	 * @see GetTrackDisplayName, GetTrackName
+	 */
+	virtual FString GetTrackLanguage(EMediaTrackType TrackType, int32 TrackIndex) const = 0;
+
+	/**
+	 * Get the internal name of the specified track.
+	 *
+	 * @param TrackType The type of track.
+	 * @param TrackIndex The index of the track.
+	 * @return Track name.
+	 * @see GetTrackDisplayName, GetTrackLanguage
+	 */
+	virtual FString GetTrackName(EMediaTrackType TrackType, int32 TrackIndex) const = 0;
+
+	/**
+	 * Get the average data rate of the specified video track.
+	 *
+	 * @param TrackIndex The index of the track.
+	 * @return Data rate (in bits per second).
+	 * @see GetVideoTrackAspectRatio, GetVideoTrackDimensions, GetVideoTrackFrameRate
+	 */
+	virtual uint32 GetVideoTrackBitRate(int32 TrackIndex) const = 0;
+
+	/**
+	 * Get the dimensions of the specified video track.
+	 *
+	 * @param TrackType The type of track.
+	 * @param TrackIndex The index of the track.
+	 * @return Video dimensions (in pixels).
+	 * @see GetVideoTrackAspectRatio, GetVideoTrackBitRate, GetVideoTrackFrameRate
+	 */
+	virtual FIntPoint GetVideoTrackDimensions(int32 TrackIndex) const = 0;
+
+	/**
+	 * Get the frame rate of the specified video track.
+	 *
+	 * @param TrackType The type of track.
+	 * @param TrackIndex The index of the track.
+	 * @return Frame rate (in frames per second).
+	 * @see GetVideoTrackAspectRatio, GetVideoTrackBitRate, GetVideoTrackDimensions
+	 */
+	virtual float GetVideoTrackFrameRate(int32 TrackIndex) const = 0;
+
+	/**
+	 * Select the active track of the given type.
+	 *
+	 * Only one track of a given type can be active at any time.
+	 *
+	 * @param TrackType The type of track to select.
+	 * @param TrackIndex The index of the track to select, or INDEX_NONE to deselect.
+	 * @return true if the track was selected, false otherwise.
+	 * @see GetNumTracks, GetSelectedTrack
+	 */
+	virtual bool SelectTrack(EMediaTrackType TrackType, int32 TrackIndex) = 0;
+
+public:
+
+	/**
+	 * Get the video's aspect ratio of the specified video track.
+	 *
+	 * @param TrackIndex Index of the video track.
+	 * @return Aspect ratio.
+	 * @see GetVideoTrackBitRate, GetVideoTrackDimensions, GetVideoTrackFrameRate
+	 */
+	float GetVideoTrackAspectRatio(int32 TrackIndex) const
+	{
+		const FIntPoint Dimensions = GetVideoTrackDimensions(TrackIndex);
+
+		if (Dimensions.Y == 0)
+		{
+			return 0.0f;
+		}
+
+		return Dimensions.X / Dimensions.Y;
+	}
+
+public:
+
+	/** Virtual destructor. */
+	virtual ~IMediaTracks() { }
+};
diff --git a/Engine/Source/Runtime/MediaAssets/MediaAssets.Build.cs b/Engine/Source/Runtime/MediaAssets/MediaAssets.Build.cs
--- a/Engine/Source/Runtime/MediaAssets/MediaAssets.Build.cs
+++ b/Engine/Source/Runtime/MediaAssets/MediaAssets.Build.cs
@@ -19,6 +19,8 @@ namespace UnrealBuildTool.Rules
 				new string[] {
                     "RenderCore",
                     "RHI",
+					"ShaderCore",
+					"UtilityShaders",
 				}
 			);
 
@@ -26,8 +28,14 @@ namespace UnrealBuildTool.Rules
 				new string[] {
 					"Runtime/MediaAssets/Private",
 					"Runtime/MediaAssets/Private/Assets",
+					"Runtime/MediaAssets/Private/Misc",
 				}
 			);
+
+			if (UEBuildConfiguration.bBuildEditor)
+			{
+				PrivateIncludePathModuleNames.Add("TargetPlatform");
+			}
 		}
 	}
 }
diff --git a/Engine/Source/Runtime/MediaAssets/Private/Assets/FileMediaSource.cpp b/Engine/Source/Runtime/MediaAssets/Private/Assets/FileMediaSource.cpp
new file mode 100644
--- /dev/null
+++ b/Engine/Source/Runtime/MediaAssets/Private/Assets/FileMediaSource.cpp
@@ -0,0 +1,42 @@
+// Copyright 1998-2016 Epic Games, Inc. All Rights Reserved.
+
+#include "MediaAssetsPCH.h"
+#include "FileMediaSource.h"
+
+
+static FName PrecacheFileName("PrecacheFile");
+
+
+bool UFileMediaSource::GetMediaOption(const FName& Key, bool DefaultValue) const
+{
+	if (Key == PrecacheFileName)
+	{
+		return PrecacheFile;
+	}
+
+	return DefaultValue;
+}
+
+
+FString UFileMediaSource::GetUrl() const
+{
+	return FString(TEXT("file://")) + GetFullPath();
+}
+
+
+bool UFileMediaSource::Validate() const
+{
+	return FPaths::FileExists(GetFullPath());
+}
+
+
+FString UFileMediaSource::GetFullPath() const
+{
+	FString FullPath = FPaths::ConvertRelativePathToFull(
+		FPaths::IsRelative(FilePath)
+			? FPaths::GameContentDir() / FilePath
+			: FilePath
+	);
+
+	return FullPath;
+}
diff --git a/Engine/Source/Runtime/MediaAssets/Private/Assets/MediaPlayer.cpp b/Engine/Source/Runtime/MediaAssets/Private/Assets/MediaPlayer.cpp
--- a/Engine/Source/Runtime/MediaAssets/Private/Assets/MediaPlayer.cpp
+++ b/Engine/Source/Runtime/MediaAssets/Private/Assets/MediaPlayer.cpp
@@ -1,8 +1,12 @@
 // Copyright 1998-2016 Epic Games, Inc. All Rights Reserved.
 
-#include "MediaAssetsPrivatePCH.h"
-#include "IMediaModule.h"
-#include "ModuleManager.h"
+#include "MediaAssetsPCH.h"
+#include "MediaCaptionSink.h"
+#include "MediaPlayer.h"
+#include "MediaPlaylist.h"
+#include "MediaSource.h"
+#include "MediaSoundWave.h"
+#include "MediaTexture.h"
 
 
 /* UMediaPlayer structors
@@ -10,8 +14,14 @@
 
 UMediaPlayer::UMediaPlayer(const FObjectInitializer& ObjectInitializer)
 	: Super(ObjectInitializer)
-	, Looping(true)
-	, Player(nullptr)
+	, PlayOnOpen(true)
+	, Shuffle(false)
+	, Loop(false)
+	, PlaylistIndex(INDEX_NONE)
+	, CaptionSink(new FMediaCaptionSink)
+#if WITH_EDITOR
+	, WasPlayingInPIE(false)
+#endif
 { }
 
 
@@ -20,95 +30,214 @@ UMediaPlayer::UMediaPlayer(const FObjectInitializer& ObjectInitializer)
 
 bool UMediaPlayer::CanPause() const
 {
-	return Player.IsValid() && Player->IsPlaying();
-}
-
-
-bool UMediaPlayer::CanPlay() const
-{
-	return Player.IsValid() && Player->IsReady();
+	return Player.IsValid() && (Player->GetControls().GetState() == EMediaState::Playing);
 }
 
 
 void UMediaPlayer::Close()
 {
+	CurrentUrl.Empty();
+	Playlist = nullptr;
+
 	if (Player.IsValid())
 	{
 		Player->Close();
 	}
-	CurrentUrl = FString();
+}
+
+
+FText UMediaPlayer::GetCaptionText() const
+{
+	return (CaptionSink != nullptr)
+		? FText::FromString(CaptionSink->GetCaption())
+		: FText::GetEmpty();
 }
 
 
 FTimespan UMediaPlayer::GetDuration() const
 {
-	if (Player.IsValid())
-	{
-		return Player->GetMediaInfo().GetDuration();
-	}
+	return Player.IsValid() ? Player->GetControls().GetDuration() : FTimespan::Zero();
+}
 
-	return FTimespan::Zero();
+
+FFloatRange UMediaPlayer::GetForwardRates(bool Unthinned)
+{
+	return Player.IsValid()
+		? Player->GetControls().GetSupportedRates(EMediaPlaybackDirections::Forward, Unthinned)
+		: FFloatRange::Empty();
+}
+
+
+int32 UMediaPlayer::GetNumTracks(EMediaPlayerTrack TrackType) const
+{
+	return Player.IsValid() ? Player->GetTracks().GetNumTracks((EMediaTrackType)TrackType) : 0;
+}
+
+
+FName UMediaPlayer::GetPlayerName() const
+{
+	return PlayerName;
 }
 
 
 float UMediaPlayer::GetRate() const
 {
-	if (Player.IsValid())
-	{
-		return Player->GetRate();
-	}
+	return Player.IsValid() ? Player->GetControls().GetRate() : 0.0f;
+}
 
-	return 0.0f;
+
+FFloatRange UMediaPlayer::GetReverseRates(bool Unthinned)
+{
+	return Player.IsValid()
+		? Player->GetControls().GetSupportedRates(EMediaPlaybackDirections::Reverse, Unthinned)
+		: FFloatRange::Empty();
+}
+
+
+int32 UMediaPlayer::GetSelectedTrack(EMediaPlayerTrack TrackType) const
+{
+	return Player.IsValid() ? Player->GetTracks().GetSelectedTrack((EMediaTrackType)TrackType) : INDEX_NONE;
 }
 
 
 FTimespan UMediaPlayer::GetTime() const
 {
-	if (Player.IsValid())
-	{
-		return Player->GetTime();
-	}
+	return Player.IsValid() ? Player->GetControls().GetTime() : FTimespan::Zero();
+}
 
-	return FTimespan::Zero();
+
+FText UMediaPlayer::GetTrackDisplayName(EMediaPlayerTrack TrackType, int32 TrackIndex) const
+{
+	return Player.IsValid() ? Player->GetTracks().GetTrackDisplayName((EMediaTrackType)TrackType, TrackIndex) : FText::GetEmpty();
 }
 
 
-const FString& UMediaPlayer::GetUrl() const
+FString UMediaPlayer::GetTrackLanguage(EMediaPlayerTrack TrackType, int32 TrackIndex) const
 {
-	return CurrentUrl;
+	return Player.IsValid() ? Player->GetTracks().GetTrackLanguage((EMediaTrackType)TrackType, TrackIndex) : FString();
 }
 
 
 bool UMediaPlayer::IsLooping() const
 {
-	return Player.IsValid() && Player->IsLooping();
+	return Player.IsValid() && Player->GetControls().IsLooping();
 }
 
 
 bool UMediaPlayer::IsPaused() const
 {
-	return Player.IsValid() && Player->IsPaused();
+	return Player.IsValid() && (Player->GetControls().GetState() == EMediaState::Paused);
 }
 
 
 bool UMediaPlayer::IsPlaying() const
 {
-	return Player.IsValid() && Player->IsPlaying();
+	return Player.IsValid() && (Player->GetControls().GetState() == EMediaState::Playing);
+}
+
+
+bool UMediaPlayer::IsPreparing() const
+{
+	return Player.IsValid() && (Player->GetControls().GetState() == EMediaState::Preparing);
 }
 
 
 bool UMediaPlayer::IsReady() const
 {
-	return Player.IsValid() && Player->IsReady();
+	if (!Player.IsValid())
+	{
+		return false;
+	}
+	
+	return ((Player->GetControls().GetState() != EMediaState::Closed) &&
+			(Player->GetControls().GetState() != EMediaState::Error) &&
+			(Player->GetControls().GetState() != EMediaState::Preparing));
 }
 
 
-bool UMediaPlayer::OpenUrl(const FString& NewUrl)
+bool UMediaPlayer::Next()
 {
-	URL = NewUrl;
-	InitializePlayer();
+	if (Playlist == nullptr)
+	{
+		return false;
+	}
+
+	int32 RemainingAttempts = Playlist->Num();
+
+	while (--RemainingAttempts >= 0)
+	{
+		UMediaSource* NextSource = Shuffle
+			? Playlist->GetRandom(PlaylistIndex)
+			: Playlist->GetNext(PlaylistIndex);
+
+		if ((NextSource != nullptr) && NextSource->Validate() && Open(NextSource->GetUrl(), *NextSource))
+		{
+			return true;
+		}
+	}
 
-	return (CurrentUrl == NewUrl);
+	return false;
+}
+
+
+bool UMediaPlayer::OpenPlaylistIndex(UMediaPlaylist* InPlaylist, int32 Index)
+{
+	if (InPlaylist == nullptr)
+	{
+		return false;
+	}
+
+	Playlist = InPlaylist;
+	PlaylistIndex = Index;
+
+	if (PlaylistIndex == INDEX_NONE)
+	{
+		return true;
+	}
+
+	UMediaSource* MediaSource = Playlist->Get(PlaylistIndex);
+
+	if (MediaSource == nullptr)
+	{
+		return false;
+	}
+	
+	if (!MediaSource->Validate())
+	{
+		UE_LOG(LogMediaAssets, Error, TEXT("Failed to validate media source %s"), *MediaSource->GetName());
+
+		return false;
+	}
+
+	return Open(MediaSource->GetUrl(), *MediaSource);
+}
+
+
+bool UMediaPlayer::OpenSource(UMediaSource* MediaSource)
+{
+	if (MediaSource == nullptr)
+	{
+		return false;
+	}
+
+	if (!MediaSource->Validate())
+	{
+		UE_LOG(LogMediaAssets, Error, TEXT("Failed to validate media source %s"), *MediaSource->GetName());
+
+		return false;
+	}
+
+	Playlist = nullptr;
+
+	return Open(MediaSource->GetUrl(), *MediaSource);
+}
+
+
+bool UMediaPlayer::OpenUrl(const FString& Url)
+{
+	Playlist = nullptr;
+
+	return Open(Url, *GetDefault<UMediaSource>());
 }
 
 
@@ -124,60 +253,217 @@ bool UMediaPlayer::Play()
 }
 
 
+bool UMediaPlayer::Previous()
+{
+	if (Playlist == nullptr)
+	{
+		return false;
+	}
+
+	int32 RemainingAttempts = Playlist->Num();
+
+	while (--RemainingAttempts >= 0)
+	{
+		UMediaSource* PrevSource = Shuffle
+			? Playlist->GetRandom(PlaylistIndex)
+			: Playlist->GetNext(PlaylistIndex);
+
+		if ((PrevSource != nullptr) && PrevSource->Validate() && Open(PrevSource->GetUrl(), *PrevSource))
+		{
+			return true;
+		}
+	}
+
+	return false;
+}
+
+
+bool UMediaPlayer::Reopen()
+{
+	if (Playlist != nullptr)
+	{
+		return OpenPlaylistIndex(Playlist, PlaylistIndex);
+	}
+
+	if (!CurrentUrl.IsEmpty())
+	{
+		OpenUrl(CurrentUrl);
+	}
+
+	return false;
+}
+
+
 bool UMediaPlayer::Rewind()
 {
 	return Seek(FTimespan::Zero());
 }
 
 
+bool UMediaPlayer::Seek(const FTimespan& InTime)
+{
+	return Player.IsValid() && Player->GetControls().Seek(InTime);
+}
+
+
+bool UMediaPlayer::SelectTrack(EMediaPlayerTrack TrackType, int32 TrackIndex)
+{
+	return Player.IsValid() ? Player->GetTracks().SelectTrack((EMediaTrackType)TrackType, TrackIndex) : false;
+}
+
+
+void UMediaPlayer::SetImageTexture(UMediaTexture* NewTexture)
+{
+	if (ImageTexture != nullptr)
+	{
+		ImageTexture->OnBeginDestroy().RemoveAll(this);
+	}
+
+	if (NewTexture != nullptr)
+	{
+		if (NewTexture == VideoTexture)
+		{
+			SetVideoTexture(nullptr);
+		}
+
+		NewTexture->OnBeginDestroy().AddUObject(this, &UMediaPlayer::HandleMediaTextureBeginDestroy);
+	}
+	
+	ImageTexture = NewTexture;
+
+	if (Player.IsValid())
+	{
+		Player->GetOutput().SetImageSink(ImageTexture);
+	}
+}
+
+
 bool UMediaPlayer::SetLooping(bool InLooping)
 {
-	return Player.IsValid() && Player->SetLooping(InLooping);
+	Loop = InLooping;
+
+	return Player.IsValid() && Player->GetControls().SetLooping(InLooping);
 }
 
 
 bool UMediaPlayer::SetRate(float Rate)
 {
-	return Player.IsValid() && Player->SetRate(Rate);
+	return Player.IsValid() && Player->GetControls().SetRate(Rate);
 }
 
 
-bool UMediaPlayer::Seek(const FTimespan& InTime)
+void UMediaPlayer::SetSoundWave(UMediaSoundWave* NewSoundWave)
 {
-	return Player.IsValid() && Player->Seek(InTime);
+	if (SoundWave != nullptr)
+	{
+		SoundWave->OnBeginDestroy().RemoveAll(this);
+	}
+
+	if (NewSoundWave != nullptr)
+	{
+		NewSoundWave->OnBeginDestroy().AddUObject(this, &UMediaPlayer::HandleMediaSoundWaveBeginDestroy);
+	}
+
+	SoundWave = NewSoundWave;
+
+	if (Player.IsValid())
+	{
+		Player->GetOutput().SetAudioSink(SoundWave);
+	}
+}
+
+
+void UMediaPlayer::SetVideoTexture(UMediaTexture* NewTexture)
+{
+	if (VideoTexture != nullptr)
+	{
+		VideoTexture->OnBeginDestroy().RemoveAll(this);
+	}
+
+	if (NewTexture != nullptr)
+	{
+		if (NewTexture == ImageTexture)
+		{
+			SetImageTexture(nullptr);
+		}
+
+		NewTexture->OnBeginDestroy().AddUObject(this, &UMediaPlayer::HandleMediaTextureBeginDestroy);
+	}
+
+	VideoTexture = NewTexture;
+
+	if (Player.IsValid())
+	{
+		Player->GetOutput().SetVideoSink(VideoTexture);
+	}
 }
 
 
 bool UMediaPlayer::SupportsRate(float Rate, bool Unthinned) const
 {
-	return Player.IsValid() && Player->GetMediaInfo().SupportsRate(Rate, Unthinned);
+	return Player.IsValid() && Player->GetControls().SupportsRate(Rate, Unthinned);
 }
 
 
 bool UMediaPlayer::SupportsScrubbing() const
 {
-	return Player.IsValid() && Player->GetMediaInfo().SupportsScrubbing();
+	return Player.IsValid() && Player->GetControls().SupportsScrubbing();
 }
 
 
 bool UMediaPlayer::SupportsSeeking() const
 {
-	return Player.IsValid() && Player->GetMediaInfo().SupportsSeeking();
+	return Player.IsValid() && Player->GetControls().SupportsSeeking();
 }
 
 
-/* UObject  overrides
+#if WITH_EDITOR
+
+void UMediaPlayer::PausePIE()
+{
+	WasPlayingInPIE = IsPlaying();
+
+	if (WasPlayingInPIE)
+	{
+		Pause();
+	}
+}
+
+
+void UMediaPlayer::ResumePIE()
+{
+	if (WasPlayingInPIE)
+	{
+		Play();
+	}
+}
+
+#endif
+
+
+/* UObject overrides
  *****************************************************************************/
 
 void UMediaPlayer::BeginDestroy()
 {
 	Super::BeginDestroy();
 
+	SetImageTexture(nullptr);
+	SetSoundWave(nullptr);
+	SetVideoTexture(nullptr);
+
 	if (Player.IsValid())
 	{
 		Player->Close();
+		Player->OnMediaEvent().RemoveAll(this);
 		Player.Reset();
 	}
+
+	if (CaptionSink != nullptr)
+	{
+		delete CaptionSink;
+		CaptionSink = nullptr;
+	}
 }
 
 
@@ -196,9 +482,19 @@ void UMediaPlayer::PostLoad()
 {
 	Super::PostLoad();
 
-	if (!HasAnyFlags(RF_ClassDefaultObject) && !GIsBuildMachine)
+	if (ImageTexture != nullptr)
+	{
+		ImageTexture->OnBeginDestroy().AddUObject(this, &UMediaPlayer::HandleMediaTextureBeginDestroy);
+	}
+
+	if (SoundWave != nullptr)
+	{
+		SoundWave->OnBeginDestroy().AddUObject(this, &UMediaPlayer::HandleMediaSoundWaveBeginDestroy);
+	}
+
+	if (VideoTexture != nullptr)
 	{
-		InitializePlayer();
+		VideoTexture->OnBeginDestroy().AddUObject(this, &UMediaPlayer::HandleMediaTextureBeginDestroy);
 	}
 }
 
@@ -207,9 +503,51 @@ void UMediaPlayer::PostLoad()
 
 void UMediaPlayer::PostEditChangeProperty(FPropertyChangedEvent& PropertyChangedEvent)
 {
+	const FName PropertyName = (PropertyChangedEvent.Property != nullptr)
+		? PropertyChangedEvent.Property->GetFName()
+		: NAME_None;
+
+	if (PropertyName == GET_MEMBER_NAME_CHECKED(UMediaPlayer, ImageTexture))
+	{
+		SetImageTexture(ImageTexture);
+	}
+	else if (PropertyName == GET_MEMBER_NAME_CHECKED(UMediaPlayer, SoundWave))
+	{
+		SetSoundWave(SoundWave);
+	}
+	else if (PropertyName == GET_MEMBER_NAME_CHECKED(UMediaPlayer, VideoTexture))
+	{
+		SetVideoTexture(VideoTexture);
+	}
+	else if (PropertyName == GET_MEMBER_NAME_CHECKED(UMediaPlayer, Loop))
+	{
+		SetLooping(Loop);
+	}
+
 	Super::PostEditChangeProperty(PropertyChangedEvent);
+}
+
+
+void UMediaPlayer::PreEditChange(UProperty* PropertyAboutToChange)
+{
+	const FName PropertyName = (PropertyAboutToChange != nullptr)
+		? PropertyAboutToChange->GetFName()
+		: NAME_None;
 
-	InitializePlayer();
+	if (PropertyName == GET_MEMBER_NAME_CHECKED(UMediaPlayer, ImageTexture))
+	{
+		SetImageTexture(nullptr);
+	}
+	else if (PropertyName == GET_MEMBER_NAME_CHECKED(UMediaPlayer, SoundWave))
+	{
+		SetSoundWave(nullptr);
+	}
+	else if (PropertyName == GET_MEMBER_NAME_CHECKED(UMediaPlayer, VideoTexture))
+	{
+		SetVideoTexture(nullptr);
+	}
+
+	Super::PreEditChange(PropertyAboutToChange);
 }
 
 #endif
@@ -218,100 +556,238 @@ void UMediaPlayer::PostEditChangeProperty(FPropertyChangedEvent& PropertyChanged
 /* UMediaPlayer implementation
  *****************************************************************************/
 
-void UMediaPlayer::InitializePlayer()
+bool UMediaPlayer::Open(const FString& Url, const IMediaOptions& Options)
 {
 	if (IsRunningDedicatedServer())
 	{
-		return;
+		return false;
 	}
 
-	if (URL == CurrentUrl)
+	if (Url.IsEmpty())
 	{
-		return;
+		return false;
 	}
 
-	// close previous player
-	CurrentUrl = FString();
+	// load media module
+	IMediaModule* MediaModule = FModuleManager::LoadModulePtr<IMediaModule>("Media");
 
-	if (Player.IsValid())
+	if (MediaModule == nullptr)
 	{
-		Player->Close();
-		Player->OnMediaEvent().RemoveAll(this);
-		Player.Reset();
+		UE_LOG(LogMediaAssets, Error, TEXT("Failed to load Media module"));
+
+		return false;
 	}
 
-	if (URL.IsEmpty())
+	TSharedPtr<IMediaPlayer> NewPlayer;
+	FName NewPlayerName = NAME_None;
+
+	// find a native player
+	if (DesiredPlayerName != NAME_None)
+	{
+		if (DesiredPlayerName != PlayerName)
+		{
+			// create desired player
+			IMediaPlayerFactory* Factory = MediaModule->GetPlayerFactory(DesiredPlayerName);
+
+			if (Factory == nullptr)
+			{
+				UE_LOG(LogMediaAssets, Error, TEXT("Could not find desired player %s for %s"), *DesiredPlayerName.ToString(), *Url);
+			}
+			else
+			{
+				NewPlayer = Factory->CreatePlayer();
+
+				if (!NewPlayer.IsValid())
+				{
+					UE_LOG(LogMediaAssets, Error, TEXT("Failed to create desired player %s for %s"), *DesiredPlayerName.ToString(), *Url);
+				}
+				else
+				{
+					NewPlayerName = DesiredPlayerName;
+				}
+			}
+		}
+		else
+		{
+			NewPlayer = Player; // reuse existing player
+		}
+	}
+	else
 	{
-		return;
+		if (PlayerName != NAME_None)
+		{
+			// try to reuse existing player
+			IMediaPlayerFactory* Factory = MediaModule->GetPlayerFactory(PlayerName);
+
+			if ((Factory != nullptr) && Factory->CanPlayUrl(Url, Options))
+			{
+				NewPlayer = Player;
+			}
+		}
+
+		if (!NewPlayer.IsValid())
+		{
+			// automatically select & create new player
+			const TArray<IMediaPlayerFactory*>& PlayerFactories = MediaModule->GetPlayerFactories();
+			const FString RunningPlatformName(FPlatformProperties::IniPlatformName());
+
+			for (IMediaPlayerFactory* Factory : PlayerFactories)
+			{
+				if (!Factory->SupportsPlatform(RunningPlatformName) || !Factory->CanPlayUrl(Url, Options))
+				{
+					continue;
+				}
+
+				NewPlayer = Factory->CreatePlayer();
+
+				if (NewPlayer.IsValid())
+				{
+					NewPlayerName = Factory->GetName();
+
+					break;
+				}
+			}
+
+			if (!NewPlayer.IsValid())
+			{
+				UE_LOG(LogMediaAssets, Error, TEXT("Could not find a native player for %s"), *Url);
+			}
+		}
 	}
 
-	// create new player
-	IMediaModule* MediaModule = FModuleManager::LoadModulePtr<IMediaModule>("Media");
+	// initialize new player
+	if (NewPlayer != Player)
+	{
+		if (Player.IsValid())
+		{
+			Player->Close();
+			Player->OnMediaEvent().RemoveAll(this);
+		}
+
+		if (NewPlayer.IsValid())
+		{
+			NewPlayer->OnMediaEvent().AddUObject(this, &UMediaPlayer::HandlePlayerMediaEvent);
+			PlayerName = NewPlayerName;
+		}
+		else
+		{
+			PlayerName = NAME_None;
+		}
+
+		Player = NewPlayer;
+	}
 
-	if (MediaModule == nullptr)
+	if (!Player.IsValid())
 	{
-		return;
+		return false;
 	}
 
-	Player = MediaModule->CreatePlayer(URL);
+	CurrentUrl = Url;
 
-	if (!Player.IsValid())
+	// open the new media source
+	if (!Player->Open(Url, Options))
 	{
-		return;
+		CurrentUrl.Empty();
+
+		return false;
 	}
 
-	Player->OnMediaEvent().AddUObject(this, &UMediaPlayer::HandlePlayerMediaEvent);
+	return true;
+}
+
+
+void UMediaPlayer::SelectDefaultTracks()
+{
+	IMediaTracks& Tracks = Player->GetTracks();
+
+	// @todo Media: consider locale when selecting default tracks
+	Tracks.SelectTrack(EMediaTrackType::Audio, 0);
+	Tracks.SelectTrack(EMediaTrackType::Caption, 0);
+	Tracks.SelectTrack(EMediaTrackType::Image, 0);
+	Tracks.SelectTrack(EMediaTrackType::Video, 0);
+}
+
 
-	// open the new media file
-	bool OpenedSuccessfully = false;
+/* UMediaPlayer callbacks
+ *****************************************************************************/
 
-	if (URL.Contains(TEXT("://")))
+void UMediaPlayer::HandleMediaSoundWaveBeginDestroy(UMediaSoundWave& DestroyedSoundWave)
+{
+	if (Player.IsValid() && (&DestroyedSoundWave == SoundWave))
 	{
-		OpenedSuccessfully = Player->Open(URL);
+		Player->GetOutput().SetAudioSink(nullptr);
 	}
-	else
-	{
-		const FString FullUrl = FPaths::ConvertRelativePathToFull(
-			FPaths::IsRelative(URL)
-				? FPaths::GameContentDir() / URL
-				: URL
-		);
+}
 
-		OpenedSuccessfully = Player->Open(FullUrl);
-	}
 
-	// finish initialization
-	if (OpenedSuccessfully)
+void UMediaPlayer::HandleMediaTextureBeginDestroy(UMediaTexture& DestroyedMediaTexture)
+{
+	if (Player.IsValid())
 	{
-		CurrentUrl = URL;
+		if (&DestroyedMediaTexture == ImageTexture)
+		{
+			Player->GetOutput().SetAudioSink(nullptr);
+		}
+		else if (&DestroyedMediaTexture == VideoTexture)
+		{
+			Player->GetOutput().SetVideoSink(nullptr);
+		}
 	}
 }
 
 
-/* UMediaPlayer callbacks
- *****************************************************************************/
-
 void UMediaPlayer::HandlePlayerMediaEvent(EMediaEvent Event)
 {
+	MediaEvent.Broadcast(Event);
+
 	switch(Event)
 	{
 	case EMediaEvent::MediaClosed:
-		MediaChangedEvent.Broadcast();
 		OnMediaClosed.Broadcast();
 		break;
 
 	case EMediaEvent::MediaOpened:
-		Player->SetLooping(Looping);
-		MediaChangedEvent.Broadcast();
-		OnMediaOpened.Broadcast(URL);
+		Player->GetControls().SetLooping(Loop);
+
+		if (ImageTexture != nullptr)
+		{
+			Player->GetOutput().SetImageSink(ImageTexture);
+		}
+
+		if (SoundWave != nullptr)
+		{
+			Player->GetOutput().SetAudioSink(SoundWave);
+		}
+
+		if (VideoTexture != nullptr)
+		{
+			Player->GetOutput().SetVideoSink(VideoTexture);
+		}
+
+		OnMediaOpened.Broadcast(CurrentUrl);
+
+		if (PlayOnOpen)
+		{
+			Play();
+		}
 		break;
 
 	case EMediaEvent::MediaOpenFailed:
-		OnMediaOpenFailed.Broadcast(URL);
+		OnMediaOpenFailed.Broadcast(CurrentUrl);
+
+		if (!Loop && (Playlist != nullptr))
+		{
+			Next();
+		}
 		break;
 
 	case EMediaEvent::PlaybackEndReached:
 		OnEndReached.Broadcast();
+
+		if (!Loop && (Playlist != nullptr))
+		{
+			Next();
+		}
 		break;
 
 	case EMediaEvent::PlaybackResumed:
@@ -323,6 +799,13 @@ void UMediaPlayer::HandlePlayerMediaEvent(EMediaEvent Event)
 		break;
 
 	case EMediaEvent::TracksChanged:
-		TracksChangedEvent.Broadcast();
+		SelectDefaultTracks();
+		break;
+	}
+
+	if ((Event == EMediaEvent::MediaOpened) ||
+		(Event == EMediaEvent::MediaOpenFailed))
+	{
+		UE_LOG(LogMediaAssets, Verbose, TEXT("Media Info:\n\n%s"), *Player->GetInfo());
 	}
 }
diff --git a/Engine/Source/Runtime/MediaAssets/Private/Assets/MediaPlaylist.cpp b/Engine/Source/Runtime/MediaAssets/Private/Assets/MediaPlaylist.cpp
new file mode 100644
--- /dev/null
+++ b/Engine/Source/Runtime/MediaAssets/Private/Assets/MediaPlaylist.cpp
@@ -0,0 +1,63 @@
+// Copyright 1998-2016 Epic Games, Inc. All Rights Reserved.
+
+#include "MediaAssetsPCH.h"
+#include "MediaPlaylist.h"
+
+
+/* UMediaPlaylist interface
+ *****************************************************************************/
+
+UMediaSource* UMediaPlaylist::Get(int32 Index)
+{
+	if (!Items.IsValidIndex(Index))
+	{
+		return nullptr;
+	}
+
+	return Items[Index];
+}
+
+
+UMediaSource* UMediaPlaylist::GetNext(int32& InOutIndex)
+{
+	if (Items.Num() == 0)
+	{
+		InOutIndex = INDEX_NONE;
+
+		return nullptr;
+	}
+
+	InOutIndex = (InOutIndex != INDEX_NONE) ? (InOutIndex + 1) % Items.Num() : 0;
+
+	return Items[InOutIndex];
+}
+
+
+UMediaSource* UMediaPlaylist::GetPrevious(int32& InOutIndex)
+{
+	if (Items.Num() == 0)
+	{
+		InOutIndex = INDEX_NONE;
+
+		return nullptr;
+	}
+
+	InOutIndex = (InOutIndex != INDEX_NONE) ? (InOutIndex + Items.Num() - 1) % Items.Num() : 0;
+
+	return Items[InOutIndex];
+}
+
+
+UMediaSource* UMediaPlaylist::GetRandom(int32& InOutIndex)
+{
+	if (Items.Num() == 0)
+	{
+		InOutIndex = INDEX_NONE;
+
+		return nullptr;
+	}
+
+	InOutIndex = FMath::RandHelper(Items.Num() - 1);
+
+	return Items[InOutIndex];
+}
diff --git a/Engine/Source/Runtime/MediaAssets/Private/Assets/MediaSoundWave.cpp b/Engine/Source/Runtime/MediaAssets/Private/Assets/MediaSoundWave.cpp
--- a/Engine/Source/Runtime/MediaAssets/Private/Assets/MediaSoundWave.cpp
+++ b/Engine/Source/Runtime/MediaAssets/Private/Assets/MediaSoundWave.cpp
@@ -1,99 +1,141 @@
 // Copyright 1998-2016 Epic Games, Inc. All Rights Reserved.
 
-#include "MediaAssetsPrivatePCH.h"
-#include "IMediaAudioTrack.h"
-#include "MediaSampleQueue.h"
+#include "MediaAssetsPCH.h"
 #include "MediaSoundWave.h"
 
 
+const int32 ZeroBufferSize = 1024;
+
+
 /* UMediaSoundWave structors
  *****************************************************************************/
 
 UMediaSoundWave::UMediaSoundWave(const FObjectInitializer& ObjectInitializer)
 	: Super(ObjectInitializer)
-	, AudioTrackIndex(INDEX_NONE)
-	, AudioQueue(MakeShareable(new FMediaSampleQueue))
+	, Paused(true)
+	, SinkNumChannels(0)
+	, SinkSampleRate(0)
 {
-	bSetupDelegates = false;
 	bLooping = false;
 	bProcedural = true;
 	Duration = INDEFINITELY_LOOPING_DURATION;
+
+	// @todo media: add support for reinitializing sound wave
+	NumChannels = 2;
+	SampleRate = 44100;
 }
 
-/* UMediaSoundWave interface
+
+/* USoundWave interface
  *****************************************************************************/
 
-TSharedPtr<IMediaPlayer> UMediaSoundWave::GetPlayer() const
+int32 UMediaSoundWave::GeneratePCMData(uint8* Data, const int32 SamplesRequested)
 {
-	if (MediaPlayer == nullptr)
+	if ((Data == nullptr) || (SamplesRequested <= 0))
 	{
-		return nullptr;
+		return 0;
 	}
 
-	return MediaPlayer->GetPlayer();
-}
-
-
-void UMediaSoundWave::SetMediaPlayer(UMediaPlayer* InMediaPlayer)
-{
-	MediaPlayer = InMediaPlayer;
+	// return queued samples
+	if (!Paused && (SinkNumChannels > 0) && (SinkSampleRate > 0) && (QueuedAudio.Num() > 0))
+	{
+		FScopeLock Lock(&CriticalSection);
 
-	InitializeTrack();
-}
+		// copy queued audio to buffer
+		if ((SinkNumChannels == NumChannels) && (SinkSampleRate == SampleRate))
+		{
+			const int32 BytesRequested = SamplesRequested * sizeof(int16);
+			const int32 BytesToCopy = FMath::Min<int32>(BytesRequested, QueuedAudio.Num());
 
+			FMemory::Memcpy((void*)Data, &QueuedAudio[0], BytesToCopy);
+			QueuedAudio.RemoveAt(0, BytesToCopy);
 
-/* USoundWave overrides
- *****************************************************************************/
+			return BytesToCopy;
+		}
 
-int32 UMediaSoundWave::GeneratePCMData(uint8* PCMData, const int32 SamplesNeeded)
-{
-	if (AudioTrack.IsValid())
-	{
-		NumChannels = AudioTrack->GetNumChannels();
-		SampleRate = AudioTrack->GetSamplesPerSecond();
-	}
+		// poor man's re-sampling
+		const float ResampleRatio = (float)SinkSampleRate / 44100.0f;
+		const int32 SamplesAvailable = (QueuedAudio.Num() / sizeof(int16)) / ResampleRatio - 2 * SinkNumChannels;
+		const int32 NumFrames = FMath::Min<int32>(SamplesRequested, SamplesAvailable) / NumChannels;
+		
+		if (NumFrames > 0)
+		{
+			int16* InputSamples = (int16*)QueuedAudio.GetData();
+			int16* OutputSamples = (int16*)Data;
+			
+			float InputFrameInterp = 0.0f;
+			int32 InputIndex = 0;
 
-	if ((NumChannels == 0) || (SampleRate == 0))
-	{
-		return 0;
-	}
+			for (int32 OutputFrame = 0; OutputFrame < NumFrames; ++OutputFrame)
+			{
+				int32 InputFrame = (int32)InputFrameInterp;
+				float Alpha = InputFrameInterp - InputFrame;
+
+				InputIndex = InputFrame * SinkNumChannels;
+
+				if (SinkNumChannels == 1)
+				{
+					// mono to stereo
+					const int16 OutputSample = (int16)FMath::Lerp((float)InputSamples[InputIndex], (float)InputSamples[InputIndex + 1], Alpha);
+
+					OutputSamples[0] = OutputSample;
+					OutputSamples[1] = OutputSample;
+				}
+				else if (SinkNumChannels == 2)
+				{
+					// stereo to stereo
+					const int16 OutputSampleL = (int16)FMath::Lerp((float)InputSamples[InputIndex], (float)InputSamples[InputIndex + 2], Alpha);
+					const int16 OutputSampleR = (int16)FMath::Lerp((float)InputSamples[InputIndex + 1], (float)InputSamples[InputIndex + 3], Alpha);
+
+					OutputSamples[0] = OutputSampleL;
+					OutputSamples[1] = OutputSampleR;
+				}
+				else if (SinkNumChannels == 6)
+				{
+					// 5.1 to stereo
+					const int16 OutputSampleC = (int16)FMath::Lerp((float)InputSamples[InputIndex], (float)InputSamples[InputIndex + 6], Alpha);
+					const int16 OutputSampleL = (int16)FMath::Lerp((float)InputSamples[InputIndex + 1], (float)InputSamples[InputIndex + 7], Alpha);
+					const int16 OutputSampleR = (int16)FMath::Lerp((float)InputSamples[InputIndex + 2], (float)InputSamples[InputIndex + 8], Alpha);
+					const int16 OutputSampleLS = (int16)FMath::Lerp((float)InputSamples[InputIndex + 3], (float)InputSamples[InputIndex + 9], Alpha);
+					const int16 OutputSampleRS = (int16)FMath::Lerp((float)InputSamples[InputIndex + 4], (float)InputSamples[InputIndex + 10], Alpha);
+					const int16 OutputSampleLFE = (int16)FMath::Lerp((float)InputSamples[InputIndex + 5], (float)InputSamples[InputIndex + 11], Alpha);
+
+					OutputSamples[0] = (OutputSampleL + OutputSampleLS) / 2 + (OutputSampleLFE + OutputSampleC) / 4;
+					OutputSamples[1] = (OutputSampleR + OutputSampleRS) / 2 + (OutputSampleLFE + OutputSampleC) / 4;
+				}
+
+				OutputSamples += NumChannels;
+				InputIndex += SinkNumChannels;
+				InputFrameInterp += ResampleRatio;
+			}
 
-	// drain media sample queue
-	TSharedPtr<TArray<uint8>, ESPMode::ThreadSafe> Sample;
+			QueuedAudio.RemoveAt(0, InputIndex * sizeof(int16));
 
-	while (AudioQueue->Dequeue(Sample))
-	{
-		QueuedAudio.Append(*Sample);
+			return (uint8*)OutputSamples - (uint8*)Data;
+		}
 	}
 
-	// get requested samples
-	if (SamplesNeeded <= 0)
+	// return zero samples if paused or buffer underrun
+	if (!Paused && (SinkNumChannels > 0) && (SinkSampleRate > 0))
 	{
-		return 0;
+		UE_LOG(LogMediaAssets, Verbose, TEXT("MediaSoundWave buffer underrun."));
 	}
 
-	const int32 SamplesAvailable = QueuedAudio.Num() / sizeof(int16);
-	const int32 BytesToCopy = FMath::Min<int32>(SamplesNeeded, SamplesAvailable) * sizeof(int16);
+	FMemory::Memzero(Data, ZeroBufferSize);
 
-	if (BytesToCopy > 0)
-	{
-		FMemory::Memcpy((void*)PCMData, &QueuedAudio[0], BytesToCopy);
-		QueuedAudio.RemoveAt(0, BytesToCopy);
-	}
-
-	return BytesToCopy;
+	return ZeroBufferSize;
 }
 
 
 FByteBulkData* UMediaSoundWave::GetCompressedData(FName Format)
 {
-	return nullptr;
+	return nullptr; // uncompressed audio
 }
 
 
 int32 UMediaSoundWave::GetResourceSizeForFormat(FName Format)
 {
-	return 0;
+	return 0; // procedural sound
 }
 
 
@@ -109,17 +151,26 @@ bool UMediaSoundWave::InitAudioResource(FName Format)
 }
 
 
-/* UObject overrides
+/* UObject interface
  *****************************************************************************/
 
+void UMediaSoundWave::BeginDestroy()
+{
+	Super::BeginDestroy();
+
+	BeginDestroyEvent.Broadcast(*this);
+}
+
+
 void UMediaSoundWave::GetAssetRegistryTags(TArray<FAssetRegistryTag>& OutTags) const
 {
+	// do nothing
 }
 
 
 SIZE_T UMediaSoundWave::GetResourceSize(EResourceSizeMode::Type Mode)
 {
-	return 0;
+	return 0; // procedural sound
 }
 
 
@@ -129,100 +180,86 @@ void UMediaSoundWave::Serialize(FArchive& Ar)
 	USoundBase::Serialize(Ar);
 }
 
-void UMediaSoundWave::PostLoad()
+
+/* IMediaAudioSink interface
+ *****************************************************************************/
+
+void UMediaSoundWave::FlushAudioSink()
 {
-	Super::PostLoad();
+	UE_LOG(LogMediaAssets, Verbose, TEXT("MediaSoundWave flushing sink."));
 
-	if (!HasAnyFlags(RF_ClassDefaultObject) && !GIsBuildMachine)
-	{
-		InitializeTrack();
-	}
+	FScopeLock Lock(&CriticalSection);
+	QueuedAudio.Empty();
 }
 
-void UMediaSoundWave::BeginDestroy()
+
+int32 UMediaSoundWave::GetAudioSinkChannels() const
 {
-	Super::BeginDestroy();
+	return SinkNumChannels;
+}
 
-	if (AudioTrack.IsValid())
-	{
-		AudioTrack->GetStream().RemoveSink(AudioQueue);
-		AudioTrack.Reset();
-	}
 
-	if (CurrentMediaPlayer.IsValid())
-	{
-		CurrentMediaPlayer->OnTracksChanged().RemoveAll(this);
-		CurrentMediaPlayer.Reset();
-	}
+int32 UMediaSoundWave::GetAudioSinkSampleRate() const
+{
+	return SinkSampleRate;
 }
 
 
-/* UMediaSoundWave implementation
- *****************************************************************************/
-
-void UMediaSoundWave::InitializeTrack()
+bool UMediaSoundWave::InitializeAudioSink(uint32 Channels, uint32 InSampleRate)
 {
-	// assign new media player asset
-	if (CurrentMediaPlayer != MediaPlayer || !bSetupDelegates)
+	UE_LOG(LogMediaAssets, Verbose, TEXT("MediaSoundWave initializing sink with %i channels at %i Hz."), Channels, InSampleRate);
+
+	if ((Channels != 1) && (Channels != 2) && (Channels != 6))
 	{
-		if (CurrentMediaPlayer != nullptr)
-		{
-			CurrentMediaPlayer->OnTracksChanged().RemoveAll(this);
-		}
+		return false; // we currently only support mono, stereo, and 5.1
+	}
 
-		CurrentMediaPlayer = MediaPlayer;
+	if ((Channels == SinkNumChannels) && (InSampleRate == SinkSampleRate))
+	{
+		return true;
+	}
 
-		if (MediaPlayer != nullptr)
-		{
-			MediaPlayer->OnTracksChanged().AddUObject(this, &UMediaSoundWave::HandleMediaPlayerTracksChanged);
-		}
+	FScopeLock Lock(&CriticalSection);
+	QueuedAudio.Empty();
 
-		bSetupDelegates = true;
-	}
+	SinkNumChannels = Channels;
+	SinkSampleRate = InSampleRate;
 
-	// disconnect from current track
-	if (AudioTrack.IsValid())
-	{
-		AudioTrack->GetStream().RemoveSink(AudioQueue);
-		AudioTrack.Reset();
+	return true;
+}
 
-		AudioQueue->Empty();
-		QueuedAudio.Empty();
-	}
 
-	// initialize from new track
-	if (MediaPlayer != nullptr)
-	{
-		IMediaPlayerPtr Player = MediaPlayer->GetPlayer();
+void UMediaSoundWave::PauseAudioSink()
+{
+	UE_LOG(LogMediaAssets, Verbose, TEXT("MediaSoundWave pausing sink."));
+	Paused = true;
+}
 
-		if (Player.IsValid())
-		{
-			auto AudioTracks = Player->GetAudioTracks();
 
-			if (AudioTracks.IsValidIndex(AudioTrackIndex))
-			{
-				AudioTrack = AudioTracks[AudioTrackIndex];
-			}
-			else if (AudioTracks.Num() > 0)
-			{
-				AudioTrack = AudioTracks[0];
-				AudioTrackIndex = 0;
-			}
-		}
-	}
+void UMediaSoundWave::PlayAudioSink(const uint8* SampleBuffer, uint32 BufferSize, FTimespan Time)
+{
+	UE_LOG(LogMediaAssets, VeryVerbose, TEXT("MediaSoundWave playing %i bytes at %s."), BufferSize, *Time.ToString());
 
-	// connect to new track
-	if (AudioTrack.IsValid())
-	{
-		AudioTrack->GetStream().AddSink(AudioQueue);
-	}
+	FScopeLock Lock(&CriticalSection);
+	QueuedAudio.Append(SampleBuffer, BufferSize);
 }
 
 
-/* UMediaSoundWave callbacks
- *****************************************************************************/
+void UMediaSoundWave::ResumeAudioSink()
+{
+	UE_LOG(LogMediaAssets, Verbose, TEXT("MediaSoundWave resuming sink."));
+	Paused = false;
+}
+
 
-void UMediaSoundWave::HandleMediaPlayerTracksChanged()
+void UMediaSoundWave::ShutdownAudioSink()
 {
-	InitializeTrack();
+	UE_LOG(LogMediaAssets, Verbose, TEXT("MediaSoundWave shutting down sink."));
+
+	FScopeLock Lock(&CriticalSection);
+	QueuedAudio.Empty();
+
+	Paused = true;
+	SinkNumChannels = 0;
+	SinkSampleRate = 0;
 }
diff --git a/Engine/Source/Runtime/MediaAssets/Private/Assets/MediaSource.cpp b/Engine/Source/Runtime/MediaAssets/Private/Assets/MediaSource.cpp
new file mode 100644
--- /dev/null
+++ b/Engine/Source/Runtime/MediaAssets/Private/Assets/MediaSource.cpp
@@ -0,0 +1,107 @@
+// Copyright 1998-2016 Epic Games, Inc. All Rights Reserved.
+
+#include "MediaAssetsPCH.h"
+#include "MediaSource.h"
+
+
+/* UMediaSource interface
+ *****************************************************************************/
+
+FString UMediaSource::GetDesiredPlayer() const
+{
+#if WITH_EDITORONLY_DATA
+	const FString RunningPlatformName(FPlatformProperties::IniPlatformName());
+	const FString* PlatformPlayer = PlatformPlayers.Find(RunningPlatformName);
+
+	if (PlatformPlayer != nullptr)
+	{
+		return *PlatformPlayer;
+	}
+#endif
+
+	return DefaultPlayer;
+}
+
+
+/* UObject interface
+ *****************************************************************************/
+
+void UMediaSource::GetAssetRegistryTags(TArray<FAssetRegistryTag>& OutTags) const
+{
+	FString Url = GetUrl();
+
+	if (!Url.IsEmpty())
+	{
+		OutTags.Add(FAssetRegistryTag("Url", Url, FAssetRegistryTag::TT_Alphabetical));
+	}
+}
+
+
+#if WITH_EDITOR
+void UMediaSource::GetAssetRegistryTagMetadata(TMap<FName, FAssetRegistryTagMetadata>& OutMetadata) const
+{
+}
+#endif
+
+
+void UMediaSource::Serialize(FArchive& Ar)
+{
+	Super::Serialize(Ar);
+
+#if WITH_EDITORONLY_DATA
+	if (Ar.IsCooking())
+	{
+		const FString* PlatformPlayer = PlatformPlayers.Find(Ar.CookingTarget()->PlatformName());
+
+		if (PlatformPlayer != nullptr)
+		{
+			DefaultPlayer = *PlatformPlayer;
+		}
+	}
+	else
+	{
+		Ar << PlatformPlayers;
+	}
+#endif
+
+	Ar << DefaultPlayer;
+}
+
+
+/* IMediaOptions interface
+ *****************************************************************************/
+
+bool UMediaSource::GetMediaOption(const FName& Key, bool DefaultValue) const
+{
+	return DefaultValue;
+}
+
+
+double UMediaSource::GetMediaOption(const FName& Key, double DefaultValue) const
+{
+	return DefaultValue;
+}
+
+
+int64 UMediaSource::GetMediaOption(const FName& Key, int64 DefaultValue) const
+{
+	return DefaultValue;
+}
+
+
+FString UMediaSource::GetMediaOption(const FName& Key, const FString& DefaultValue) const
+{
+	return DefaultValue;
+}
+
+
+FText UMediaSource::GetMediaOption(const FName& Key, const FText& DefaultValue) const
+{
+	return DefaultValue;
+}
+
+
+bool UMediaSource::HasMediaOption(const FName& Key) const
+{
+	return false;
+}
diff --git a/Engine/Source/Runtime/MediaAssets/Private/Assets/MediaTexture.cpp b/Engine/Source/Runtime/MediaAssets/Private/Assets/MediaTexture.cpp
--- a/Engine/Source/Runtime/MediaAssets/Private/Assets/MediaTexture.cpp
+++ b/Engine/Source/Runtime/MediaAssets/Private/Assets/MediaTexture.cpp
@@ -1,8 +1,8 @@
 // Copyright 1998-2016 Epic Games, Inc. All Rights Reserved.
 
-#include "MediaAssetsPrivatePCH.h"
-#include "IMediaVideoTrack.h"
-#include "MediaSampleBuffer.h"
+#include "MediaAssetsPCH.h"
+#include "MediaTexture.h"
+#include "MediaTextureResource.h"
 
 
 /* UMediaTexture structors
@@ -10,264 +10,242 @@
 
 UMediaTexture::UMediaTexture(const FObjectInitializer& ObjectInitializer)
 	: Super(ObjectInitializer)
+	, AddressX(TA_Clamp)
+	, AddressY(TA_Clamp)
 	, ClearColor(FLinearColor::Black)
-	, VideoTrackIndex(INDEX_NONE)
-	, MediaPlayer(nullptr)
-	, CurrentMediaPlayer(nullptr)
-	, VideoBuffer(MakeShareable(new FMediaSampleBuffer))
+	, SinkDimensions(FIntPoint(1, 1))
+	, SinkFormat(EMediaTextureSinkFormat::CharBGRA)
+	, SinkMode(EMediaTextureSinkMode::Unbuffered)
 {
-	bDelegatesAdded = false;
 	NeverStream = true;
-	UpdateResource();
 }
 
 
-/* UMediaTexture interface
+/* FTickerObjectBase interface
  *****************************************************************************/
 
-TSharedPtr<class IMediaPlayer> UMediaTexture::GetPlayer() const
+bool UMediaTexture::Tick(float DeltaTime)
 {
-	return (MediaPlayer != nullptr)
-		? MediaPlayer->GetPlayer()
-		: nullptr;
-}
-
+	// process deferred tasks
+	TFunction<void()> Task;
 
-void UMediaTexture::SetMediaPlayer(UMediaPlayer* InMediaPlayer)
-{
-	MediaPlayer = InMediaPlayer;
+	while (GameThreadTasks.Dequeue(Task))
+	{
+		Task();
+	}
 
-	InitializeTrack();
+	return true;
 }
 
 
-/* UTexture  overrides
+/* IMediaTextureSink interface
  *****************************************************************************/
 
-FTextureResource* UMediaTexture::CreateResource()
+void* UMediaTexture::AcquireTextureSinkBuffer()
 {
-	return new FMediaTextureResource(this, VideoBuffer);
+	FScopeLock Lock(&CriticalSection);
+
+	return (Resource != nullptr) ? ((FMediaTextureResource*)Resource)->AcquireBuffer() : nullptr;
 }
 
 
-EMaterialValueType UMediaTexture::GetMaterialType()
+void UMediaTexture::DisplayTextureSinkBuffer(FTimespan /*Time*/)
 {
-	return MCT_Texture2D;
+	FScopeLock Lock(&CriticalSection);
+
+	if (Resource != nullptr)
+	{
+		((FMediaTextureResource*)Resource)->DisplayBuffer();
+	}
 }
 
 
-float UMediaTexture::GetSurfaceWidth() const
+FIntPoint UMediaTexture::GetTextureSinkDimensions() const
 {
-	return CachedDimensions.X;
+	FScopeLock Lock(&CriticalSection);
+
+	return (Resource != nullptr) ? ((FMediaTextureResource*)Resource)->GetSizeXY() : FIntPoint::ZeroValue;
 }
 
 
-float UMediaTexture::GetSurfaceHeight() const
+EMediaTextureSinkFormat UMediaTexture::GetTextureSinkFormat() const
 {
-	return CachedDimensions.Y;
+	return SinkFormat;
 }
 
-void UMediaTexture::UpdateResource()
+
+EMediaTextureSinkMode UMediaTexture::GetTextureSinkMode() const
 {
-#if WITH_ENGINE
-	if (VideoTrack.IsValid() && Resource)
-	{		
-		VideoTrack->UnbindTexture(Resource->TextureRHI.GetReference());
-	}
-#endif
-	UTexture::UpdateResource();
+	return SinkMode;
 }
 
 
-/* UObject  overrides
- *****************************************************************************/
-
-void UMediaTexture::BeginDestroy()
+FRHITexture* UMediaTexture::GetTextureSinkTexture()
 {
-	Super::BeginDestroy();
-
-	// synchronize with the rendering thread by inserting a fence
- 	if (!ReleasePlayerFence)
- 	{
- 		ReleasePlayerFence = new FRenderCommandFence();
- 	}
+	FScopeLock Lock(&CriticalSection);
 
- 	ReleasePlayerFence->BeginFence();
+	return (Resource != nullptr) ? ((FMediaTextureResource*)Resource)->GetTexture() : nullptr;
 }
 
 
-void UMediaTexture::FinishDestroy()
+bool UMediaTexture::InitializeTextureSink(FIntPoint Dimensions, EMediaTextureSinkFormat Format, EMediaTextureSinkMode Mode)
 {
-	delete ReleasePlayerFence;
-	ReleasePlayerFence = nullptr;
+	UE_LOG(LogMediaAssets, Verbose, TEXT("MediaTexture initializing sink with %i x %i pixels %s."), Dimensions.X, Dimensions.Y, (Mode == EMediaTextureSinkMode::Buffered) ? TEXT("Buffered") : TEXT("Unbuffered"));
 
-	if (VideoTrack.IsValid())
-	{
-		VideoTrack->GetStream().RemoveSink(VideoBuffer);
-		VideoTrack.Reset();
-	}
+	SinkDimensions = Dimensions;
+	SinkFormat = Format;
+	SinkMode = Mode;
+
+	FScopeLock Lock(&CriticalSection);
 
-	if (CurrentMediaPlayer.IsValid())
+	if (Resource == nullptr)
 	{
-		CurrentMediaPlayer->OnTracksChanged().RemoveAll(this);
-		CurrentMediaPlayer.Reset();
+		return false;
 	}
 
-	Super::FinishDestroy();
+	((FMediaTextureResource*)Resource)->InitializeBuffer(Dimensions, Format, Mode);
+
+	return true;
 }
 
 
-FString UMediaTexture::GetDesc()
+void UMediaTexture::ReleaseTextureSinkBuffer()
 {
-	TSharedPtr<IMediaPlayer> MediaPlayerPtr = GetPlayer();
+	FScopeLock Lock(&CriticalSection);
 
-	if (!MediaPlayerPtr.IsValid())
+	if (Resource != nullptr)
 	{
-		return FString();
+		((FMediaTextureResource*)Resource)->ReleaseBuffer();
 	}
+}
+
 
-	return FString::Printf(TEXT("%dx%d [%s]"), CachedDimensions.X, CachedDimensions.Y, GPixelFormats[GetFormat()].Name);
+void UMediaTexture::ShutdownTextureSink()
+{
+	if (ClearColor.A != 0.0f)
+	{
+		SinkDimensions = FIntPoint(1, 1);
+		SinkFormat = EMediaTextureSinkFormat::CharBGRA;
+		SinkMode = EMediaTextureSinkMode::Unbuffered;
+
+		GameThreadTasks.Enqueue([=]() {
+			UpdateResource();
+		});
+	}
 }
 
 
-SIZE_T UMediaTexture::GetResourceSize(EResourceSizeMode::Type Mode)
+bool UMediaTexture::SupportsTextureSinkFormat(EMediaTextureSinkFormat Format) const
 {
-	return CachedDimensions.X * CachedDimensions.Y * 4;
+	return true; // all formats are supported
 }
 
 
-bool UMediaTexture::IsReadyForFinishDestroy()
+void UMediaTexture::UpdateTextureSinkBuffer(const uint8* Data, uint32 Pitch)
 {
-	// ready to call FinishDestroy if the flushing fence has been hit
-	return (Super::IsReadyForFinishDestroy() && ReleasePlayerFence && ReleasePlayerFence->IsFenceComplete());
+	FScopeLock Lock(&CriticalSection);
+
+	if (Resource != nullptr)
+	{
+		((FMediaTextureResource*)Resource)->UpdateBuffer(Data, Pitch);
+	}
 }
 
 
-void UMediaTexture::PostLoad()
+void UMediaTexture::UpdateTextureSinkResource(FRHITexture* RenderTarget, FRHITexture* ShaderResource)
 {
-	Super::PostLoad();
+	FScopeLock Lock(&CriticalSection);
 
-	if (!HasAnyFlags(RF_ClassDefaultObject) && !GIsBuildMachine)
+	if (Resource != nullptr)
 	{
-		InitializeTrack();
+		((FMediaTextureResource*)Resource)->UpdateTextures(RenderTarget, ShaderResource);
 	}
 }
 
 
-#if WITH_EDITOR
+/* UTexture interface
+ *****************************************************************************/
 
-void UMediaTexture::PreEditChange(UProperty* PropertyAboutToChange)
+FTextureResource* UMediaTexture::CreateResource()
 {
-	// this will release the FMediaTextureResource
-	Super::PreEditChange(PropertyAboutToChange);
+	return new FMediaTextureResource(*this, ClearColor, SinkDimensions, SinkFormat, SinkMode);
+}
 
-	FlushRenderingCommands();
+
+EMaterialValueType UMediaTexture::GetMaterialType()
+{
+	return MCT_Texture2D;
 }
 
 
-void UMediaTexture::PostEditChangeProperty(FPropertyChangedEvent& PropertyChangedEvent)
+float UMediaTexture::GetSurfaceWidth() const
 {
-	InitializeTrack();
+	return (Resource != nullptr) ? Resource->GetSizeX() : 0.0f;
+}
+
 
-	// this will recreate the FMediaTextureResource
-	Super::PostEditChangeProperty(PropertyChangedEvent);
+float UMediaTexture::GetSurfaceHeight() const
+{
+	return (Resource != nullptr) ? Resource->GetSizeY() : 0.0f;
 }
 
-#endif // WITH_EDITOR
+
+void UMediaTexture::UpdateResource()
+{
+	FScopeLock Lock(&CriticalSection);
+
+	Super::UpdateResource();
+}
 
 
-/* UMediaTexture implementation
+/* UObject interface
  *****************************************************************************/
 
-void UMediaTexture::InitializeTrack()
+void UMediaTexture::BeginDestroy()
 {
-	// assign new media player asset
-	if ((CurrentMediaPlayer != MediaPlayer) || !bDelegatesAdded)
-	{
-		if (CurrentMediaPlayer != nullptr)
-		{
-			CurrentMediaPlayer->OnTracksChanged().RemoveAll(this);
-		}
+	Super::BeginDestroy();
 
-		CurrentMediaPlayer = MediaPlayer;
+	BeginDestroyEvent.Broadcast(*this);
+}
 
-		if (MediaPlayer != nullptr)
-		{
-			MediaPlayer->OnTracksChanged().AddUObject(this, &UMediaTexture::HandleMediaPlayerTracksChanged);
-		}
 
-		bDelegatesAdded = true;
-	}
+FString UMediaTexture::GetDesc()
+{
+	return FString::Printf(TEXT("%dx%d [%s]"), GetSurfaceWidth(),  GetSurfaceHeight(), GPixelFormats[PF_B8G8R8A8].Name);
+}
 
-	// disconnect from current track
-	if (VideoTrack.IsValid())
-	{
-		VideoTrack->GetStream().RemoveSink(VideoBuffer);
 
-#if WITH_ENGINE
-		if ((Resource != nullptr) && Resource->TextureRHI.IsValid())
-		{
-			VideoTrack->UnbindTexture(Resource->TextureRHI.GetReference());
-		}
-#endif
+SIZE_T UMediaTexture::GetResourceSize(EResourceSizeMode::Type Mode)
+{
+	return (Resource != nullptr) ? ((FMediaTextureResource*)Resource)->GetResourceSize() : 0;
+}
 
-		VideoTrack.Reset();
-	}
 
-	// initialize from new track
-	if (MediaPlayer != nullptr)
-	{
-		IMediaPlayerPtr Player = MediaPlayer->GetPlayer();
-
-		if (Player.IsValid())
-		{
-			auto VideoTracks = Player->GetVideoTracks();
-
-			if (VideoTracks.IsValidIndex(VideoTrackIndex))
-			{
-				VideoTrack = VideoTracks[VideoTrackIndex];
-			}
-			else if (VideoTracks.Num() > 0)
-			{
-				VideoTrack = VideoTracks[0];
-				VideoTrackIndex = 0;
-			}
-		}
-	}
+#if WITH_EDITOR
 
-	if (VideoTrack.IsValid())
+void UMediaTexture::PostEditChangeProperty(FPropertyChangedEvent& PropertyChangedEvent)
+{
+	const FName PropertyName = (PropertyChangedEvent.Property != nullptr)
+		? PropertyChangedEvent.Property->GetFName()
+		: NAME_None;
+
+	if ((PropertyName == GET_MEMBER_NAME_CHECKED(UMediaTexture, ClearColor)) &&
+		(PropertyChangedEvent.ChangeType != EPropertyChangeType::ValueSet))
 	{
-		CachedDimensions = VideoTrack->GetDimensions();
+		UObject::PostEditChangeProperty(PropertyChangedEvent);
 	}
 	else
 	{
-		CachedDimensions = FIntPoint(ForceInit);
-	}
-
-	UpdateResource();
-
-	// connect to new track
-	if (VideoTrack.IsValid())
-	{
-		VideoTrack->GetStream().AddSink(VideoBuffer);
-
-#if WITH_ENGINE
-		FlushRenderingCommands();
-
-		if ((Resource != nullptr) && Resource->TextureRHI.IsValid())
-		{
-			IMediaVideoTrack* VideoTrackPtr = static_cast<IMediaVideoTrack*>(VideoTrack.Get());
-			VideoTrackPtr->BindTexture((Resource->TextureRHI.GetReference()));
-		}
-#endif
+		Super::PostEditChangeProperty(PropertyChangedEvent);
 	}
 }
 
 
-/* UMediaTexture callbacks
- *****************************************************************************/
-
-void UMediaTexture::HandleMediaPlayerTracksChanged()
+void UMediaTexture::PreEditChange(UProperty* PropertyAboutToChange)
 {
-	InitializeTrack();
+	Super::PreEditChange(PropertyAboutToChange);
+
+	FlushRenderingCommands();
 }
+
+#endif // WITH_EDITOR
diff --git a/Engine/Source/Runtime/MediaAssets/Private/Assets/PlatformMediaSource.cpp b/Engine/Source/Runtime/MediaAssets/Private/Assets/PlatformMediaSource.cpp
new file mode 100644
--- /dev/null
+++ b/Engine/Source/Runtime/MediaAssets/Private/Assets/PlatformMediaSource.cpp
@@ -0,0 +1,67 @@
+// Copyright 1998-2016 Epic Games, Inc. All Rights Reserved.
+
+#include "MediaAssetsPCH.h"
+#include "MediaSource.h"
+#include "PlatformMediaSource.h"
+
+
+/* UMediaSource interface
+ *****************************************************************************/
+
+FString UPlatformMediaSource::GetUrl() const
+{
+	return (DefaultSource != nullptr) ? DefaultSource->GetUrl() : FString();
+}
+
+
+void UPlatformMediaSource::Serialize(FArchive& Ar)
+{
+	Super::Serialize(Ar);
+
+
+}
+
+
+bool UPlatformMediaSource::Validate() const
+{
+	return (DefaultSource != nullptr) ? DefaultSource->Validate() : false;
+}
+
+
+/* IMediaOptions interface
+ *****************************************************************************/
+
+bool UPlatformMediaSource::GetMediaOption(const FName& Key, bool DefaultValue) const
+{
+	return (DefaultSource != nullptr) ? DefaultSource->GetMediaOption(Key, DefaultValue) : DefaultValue;
+}
+
+
+double UPlatformMediaSource::GetMediaOption(const FName& Key, double DefaultValue) const
+{
+	return (DefaultSource != nullptr) ? DefaultSource->GetMediaOption(Key, DefaultValue) : DefaultValue;
+}
+
+
+int64 UPlatformMediaSource::GetMediaOption(const FName& Key, int64 DefaultValue) const
+{
+	return (DefaultSource != nullptr) ? DefaultSource->GetMediaOption(Key, DefaultValue) : DefaultValue;
+}
+
+
+FString UPlatformMediaSource::GetMediaOption(const FName& Key, const FString& DefaultValue) const
+{
+	return (DefaultSource != nullptr) ? DefaultSource->GetMediaOption(Key, DefaultValue) : DefaultValue;
+}
+
+
+FText UPlatformMediaSource::GetMediaOption(const FName& Key, const FText& DefaultValue) const
+{
+	return (DefaultSource != nullptr) ? DefaultSource->GetMediaOption(Key, DefaultValue) : DefaultValue;
+}
+
+
+bool UPlatformMediaSource::HasMediaOption(const FName& Key) const
+{
+	return (DefaultSource != nullptr) ? DefaultSource->HasMediaOption(Key) : false;
+}
diff --git a/Engine/Source/Runtime/MediaAssets/Private/Assets/StreamMediaSource.cpp b/Engine/Source/Runtime/MediaAssets/Private/Assets/StreamMediaSource.cpp
new file mode 100644
--- /dev/null
+++ b/Engine/Source/Runtime/MediaAssets/Private/Assets/StreamMediaSource.cpp
@@ -0,0 +1,16 @@
+// Copyright 1998-2016 Epic Games, Inc. All Rights Reserved.
+
+#include "MediaAssetsPCH.h"
+#include "StreamMediaSource.h"
+
+
+FString UStreamMediaSource::GetUrl() const
+{
+	return StreamUrl;
+}
+
+
+bool UStreamMediaSource::Validate() const
+{
+	return StreamUrl.Contains(TEXT("://"));
+}
diff --git a/Engine/Source/Runtime/MediaAssets/Private/MediaAssetsModule.cpp b/Engine/Source/Runtime/MediaAssets/Private/MediaAssetsModule.cpp
--- a/Engine/Source/Runtime/MediaAssets/Private/MediaAssetsModule.cpp
+++ b/Engine/Source/Runtime/MediaAssets/Private/MediaAssetsModule.cpp
@@ -1,9 +1,13 @@
 // Copyright 1998-2016 Epic Games, Inc. All Rights Reserved.
 
-#include "MediaAssetsPrivatePCH.h"
+#include "MediaAssetsPCH.h"
+#include "MediaPlayer.h"
 #include "ModuleInterface.h"
 
 
+DEFINE_LOG_CATEGORY(LogMediaAssets);
+
+
 /**
  * Implements the MediaAssets module.
  */
@@ -13,7 +17,7 @@ class FMediaAssetsModule
 {
 public:
 
-	// FSelfRegisteringExec interface
+	//~ FSelfRegisteringExec interface
 
 	virtual bool Exec(UWorld* InWorld, const TCHAR* Cmd, FOutputDevice& Ar) override
 	{
@@ -46,7 +50,7 @@ public:
 
 public:
 
-	// IModuleInterface interface
+	//~ IModuleInterface interface
 
 	virtual void StartupModule() override { }
 	virtual void ShutdownModule() override { }
diff --git a/Engine/Source/Runtime/MediaAssets/Private/MediaAssetsPCH.h b/Engine/Source/Runtime/MediaAssets/Private/MediaAssetsPCH.h
new file mode 100644
--- /dev/null
+++ b/Engine/Source/Runtime/MediaAssets/Private/MediaAssetsPCH.h
@@ -0,0 +1,28 @@
+// Copyright 1998-2016 Epic Games, Inc. All Rights Reserved.
+
+#pragma once
+
+#include "Runtime/Core/Public/Core.h"
+#include "Runtime/Core/Public/Containers/TripleBuffer.h"
+#include "Runtime/Core/Public/Modules/ModuleManager.h"
+#include "Runtime/CoreUObject/Public/CoreUObject.h"
+#include "Runtime/Engine/Classes/DeviceProfiles/DeviceProfile.h"
+#include "Runtime/Engine/Public/Engine.h"
+#include "Runtime/Engine/Public/GlobalShader.h"
+#include "Runtime/Media/Public/IMediaControls.h"
+#include "Runtime/Media/Public/IMediaModule.h"
+#include "Runtime/Media/Public/IMediaOptions.h"
+#include "Runtime/Media/Public/IMediaOutput.h"
+#include "Runtime/Media/Public/IMediaPlayer.h"
+#include "Runtime/Media/Public/IMediaPlayerFactory.h"
+#include "Runtime/Media/Public/IMediaTracks.h"
+#include "Runtime/RenderCore/Public/RenderingThread.h"
+#include "Runtime/UtilityShaders/Public/MediaShaders.h"
+
+#if WITH_EDITORONLY_DATA
+	#include "Developer/TargetPlatform/Public/TargetPlatform.h"
+#endif
+
+
+/** Declares a log category for this module. */
+DECLARE_LOG_CATEGORY_EXTERN(LogMediaAssets, Log, All);
diff --git a/Engine/Source/Runtime/MediaAssets/Private/Misc/MediaCaptionSink.h b/Engine/Source/Runtime/MediaAssets/Private/Misc/MediaCaptionSink.h
new file mode 100644
--- /dev/null
+++ b/Engine/Source/Runtime/MediaAssets/Private/Misc/MediaCaptionSink.h
@@ -0,0 +1,68 @@
+// Copyright 1998-2016 Epic Games, Inc. All Rights Reserved.
+
+#pragma once
+
+#include "IMediaStringSink.h"
+
+
+class FMediaCaptionSink
+	: public FTickerObjectBase
+	, public IMediaStringSink
+{
+public:
+
+	/**
+	 * Get the current caption/subtitle string, if any.
+	 *
+	 * @return Caption string.
+	 */
+	FString GetCaption() const
+	{
+		return Caption;
+	}
+
+public:
+
+	//~ FTickerObjectBase interface
+
+	virtual bool Tick(float DeltaTime) override
+	{
+		TFunction<void()> Task;
+
+		while (GameThreadTasks.Dequeue(Task))
+		{
+			Task();
+		}
+
+		return true;
+	}
+
+public:
+
+	//~ IMediaStringSink interface
+
+	virtual bool InitializeStringSink() override
+	{
+		return true;
+	}
+
+	virtual void DisplayStringSinkString(const FString& String, FTimespan Time) override
+	{
+		GameThreadTasks.Enqueue([=]() {
+			Caption = String;
+		});
+	}
+
+	virtual void ShutdownStringSink() override
+	{
+		// do nothing
+	}
+
+private:
+
+	/** The current caption string. */
+	FString Caption;
+
+	/** Asynchronous tasks for the game thread. */
+	TQueue<TFunction<void()>> GameThreadTasks;
+};
diff --git a/Engine/Source/Runtime/MediaAssets/Private/Misc/MediaTextureResource.cpp b/Engine/Source/Runtime/MediaAssets/Private/Misc/MediaTextureResource.cpp
new file mode 100644
--- /dev/null
+++ b/Engine/Source/Runtime/MediaAssets/Private/Misc/MediaTextureResource.cpp
@@ -0,0 +1,729 @@
+// Copyright 1998-2016 Epic Games, Inc. All Rights Reserved.
+
+#include "MediaAssetsPCH.h"
+#include "MediaTexture.h"
+#include "MediaTextureResource.h"
+
+
+/* FMediaTextureResource structors
+ *****************************************************************************/
+
+FMediaTextureResource::FMediaTextureResource(UMediaTexture& InOwner, const FLinearColor& InClearColor, const FIntPoint& InOutputDimensions, EMediaTextureSinkFormat InSinkFormat, EMediaTextureSinkMode InSinkMode)
+	: Owner(InOwner)
+	, TripleBuffer(BufferResources)
+	, BufferBytesPerPixel(0)
+	, BufferClearColor(InClearColor)
+	, BufferDimensions(FIntPoint::ZeroValue)
+	, OutputDimensions(InOutputDimensions)
+	, RequiresConversion(false)
+	, SinkFormat(InSinkFormat)
+	, SinkMode(InSinkMode)
+	, State(EState::ShutDown)
+{ }
+
+
+FMediaTextureResource::~FMediaTextureResource()
+{
+	FlushRenderingCommands();
+}
+
+
+/* FMediaTextureResource interface
+ *****************************************************************************/
+
+void* FMediaTextureResource::AcquireBuffer()
+{
+	if (State != EState::Initialized)
+	{
+		return nullptr;
+	}
+
+	if (SinkMode == EMediaTextureSinkMode::Buffered)
+	{
+		return TripleBuffer.GetWriteBuffer().LockedData;
+	}
+
+	check(IsInActualRenderingThread());
+
+	uint32 OutStride = 0;
+	TRefCountPtr<FRHITexture2D>& Texture = RequiresConversion
+		? BufferResources[2].RenderTarget
+		: RenderTargetTextureRHI;
+
+	return RHILockTexture2D(Texture.GetReference(), 0, RLM_WriteOnly, OutStride, false /*bLockWithinMiptail*/, false /*bFlushRHIThread*/);
+}
+
+
+void FMediaTextureResource::DisplayBuffer()
+{
+	if (State != EState::Initialized)
+	{
+		return;
+	}
+
+	if (IsInActualRenderingThread())
+	{
+		SwapResource();
+	}
+	else
+	{
+		RenderThreadTasks.Enqueue([this]() {
+			SwapResource();
+		});
+	}
+}
+
+
+SIZE_T FMediaTextureResource::GetResourceSize() const
+{
+	SIZE_T ResourceSize = 0;
+
+	if (OutputDimensions.GetMin() > 0)
+	{
+		const SIZE_T BufferSize = BufferDimensions.X * BufferDimensions.Y * BufferBytesPerPixel;
+
+		if (SinkMode == EMediaTextureSinkMode::Buffered)
+		{
+			ResourceSize += 3 * BufferSize;
+		}
+		else
+		{
+			ResourceSize += BufferSize;
+		}
+
+		if (RequiresConversion)
+		{
+			ResourceSize += OutputDimensions.X * OutputDimensions.Y * 4;
+		}
+	}
+
+	return ResourceSize;
+}
+
+
+FRHITexture* FMediaTextureResource::GetTexture()
+{
+	if (State != EState::Initialized)
+	{
+		return nullptr;
+	}
+
+	check(IsInRenderingThread());
+	check(SinkMode == EMediaTextureSinkMode::Unbuffered);
+
+	return RenderTargetTextureRHI;
+}
+
+
+void FMediaTextureResource::InitializeBuffer(FIntPoint Dimensions, EMediaTextureSinkFormat Format, EMediaTextureSinkMode Mode)
+{
+	State = EState::Initializing;
+
+	if (IsInActualRenderingThread())
+	{
+		InitializeResource(Dimensions, Format, Mode);
+	}
+	else
+	{
+		RenderThreadTasks.Enqueue([=]() {
+			InitializeResource(Dimensions, Format, Mode);
+		});
+	}
+}
+
+
+void FMediaTextureResource::ReleaseBuffer()
+{
+	if (State != EState::Initialized)
+	{
+		return;
+	}
+
+	if (SinkMode == EMediaTextureSinkMode::Buffered)
+	{
+		TripleBuffer.SwapWriteBuffers();
+	}
+	else
+	{
+		check(IsInActualRenderingThread());
+
+		TRefCountPtr<FRHITexture2D>& Texture = RequiresConversion
+			? BufferResources[2].RenderTarget
+			: RenderTargetTextureRHI;
+
+		RHIUnlockTexture2D(Texture.GetReference(), 0, false /*bLockWithinMiptail*/);
+	}
+}
+
+
+void FMediaTextureResource::ShutdownBuffer()
+{
+	if ((State == EState::ShutDown) || (State == EState::ShuttingDown))
+	{
+		return;
+	}
+
+	State = EState::ShuttingDown;
+
+	if (IsInActualRenderingThread())
+	{
+		ReleaseDynamicRHI();
+	}
+	else
+	{
+		RenderThreadTasks.Enqueue([this]() {
+			ReleaseDynamicRHI();
+		});
+	}
+}
+
+
+void FMediaTextureResource::UpdateBuffer(const uint8* Data, uint32 Pitch)
+{
+	if (State != EState::Initialized)
+	{
+		return;
+	}
+
+	// calculate pitch
+	const SIZE_T BytesPerRow = OutputDimensions.X * BufferBytesPerPixel;
+
+	if (Pitch == 0)
+	{
+		Pitch = BytesPerRow;
+	}
+
+	// update sink
+	if (SinkMode == EMediaTextureSinkMode::Buffered)
+	{
+		uint8* Dest = (uint8*)TripleBuffer.GetWriteBuffer().LockedData;
+
+		if (Dest != nullptr)
+		{
+			if (Pitch != BytesPerRow)
+			{
+				const uint8* Row = Data;
+
+				for (int32 RowIndex = 0; RowIndex < BufferDimensions.Y; ++RowIndex)
+				{
+					FMemory::Memcpy(Dest, Row, BytesPerRow);
+					Dest += BytesPerRow;
+					Row += Pitch;
+				}
+			}
+			else
+			{
+				FMemory::Memcpy(Dest, Data, BufferDimensions.Y * BytesPerRow);
+			}
+
+			TripleBuffer.SwapWriteBuffers();
+		}
+	}
+	else if (IsInActualRenderingThread())
+	{
+		if (RequiresConversion)
+		{
+			FUpdateTextureRegion2D Region(0, 0, 0, 0, BufferDimensions.X, BufferDimensions.Y);
+			RHIUpdateTexture2D(BufferResources[2].RenderTarget.GetReference(), 0, Region, Pitch, Data);
+			ConvertResource(BufferResources[2]);
+		}
+		else
+		{
+			FUpdateTextureRegion2D Region(0, 0, 0, 0, OutputDimensions.X, OutputDimensions.Y);
+			RHIUpdateTexture2D(RenderTargetTextureRHI.GetReference(), 0, Region, Pitch, Data);
+		}
+	}
+	else
+	{
+		UE_LOG(LogMediaAssets, Error, TEXT("UpdateTextureSinkBuffer must be called on render thread in Unbuffered mode."));
+	}
+}
+
+
+void FMediaTextureResource::UpdateTextures(FRHITexture* RenderTarget, FRHITexture* ShaderResource)
+{
+	if (State != EState::Initialized)
+	{
+		return;
+	}
+
+	check(SinkMode == EMediaTextureSinkMode::Unbuffered);
+
+	FResource Resource;
+	{
+		Resource.RenderTarget = (RenderTarget != nullptr) ? RenderTarget->GetTexture2D() : nullptr;
+		Resource.ShaderResource = (ShaderResource != nullptr) ? ShaderResource->GetTexture2D() : nullptr;
+	}
+
+	if (IsInActualRenderingThread())
+	{
+		SetResource(Resource);
+	}
+	else
+	{
+		RenderThreadTasks.Enqueue([=]() {
+			SetResource(Resource);
+		});
+	}
+}
+
+
+/* FRenderTarget interface
+ *****************************************************************************/
+
+FIntPoint FMediaTextureResource::GetSizeXY() const
+{
+	return OutputDimensions;
+}
+
+
+/* FTextureResource interface
+ *****************************************************************************/
+
+FString FMediaTextureResource::GetFriendlyName() const
+{
+	return Owner.GetPathName();
+}
+
+
+uint32 FMediaTextureResource::GetSizeX() const
+{
+	return OutputDimensions.X;
+}
+
+
+uint32 FMediaTextureResource::GetSizeY() const
+{
+	return OutputDimensions.Y;
+}
+
+
+void FMediaTextureResource::InitDynamicRHI()
+{
+	// create the sampler state
+	FSamplerStateInitializerRHI SamplerStateInitializer(
+		(ESamplerFilter)UDeviceProfileManager::Get().GetActiveProfile()->GetTextureLODSettings()->GetSamplerFilter(&Owner),
+		(Owner.AddressX == TA_Wrap) ? AM_Wrap : ((Owner.AddressX == TA_Clamp) ? AM_Clamp : AM_Mirror),
+		(Owner.AddressY == TA_Wrap) ? AM_Wrap : ((Owner.AddressY == TA_Clamp) ? AM_Clamp : AM_Mirror),
+		AM_Wrap
+	);
+
+	SamplerStateRHI = RHICreateSamplerState(SamplerStateInitializer);
+
+	if (OutputDimensions.GetMin() <= 0)
+	{
+		return;
+	}
+
+	// determine buffer resource pixel format
+	EPixelFormat BufferFormat;
+	{
+		switch (SinkFormat)
+		{
+		case EMediaTextureSinkFormat::CharNV12:
+		case EMediaTextureSinkFormat::CharNV21:
+			BufferBytesPerPixel = 1;
+			BufferDimensions = FIntPoint(OutputDimensions.X, OutputDimensions.Y * 1.5f);
+			BufferFormat = PF_G8;
+			RequiresConversion = true;
+			break;
+
+		case EMediaTextureSinkFormat::CharUYVY:
+		case EMediaTextureSinkFormat::CharYUY2:
+		case EMediaTextureSinkFormat::CharYVYU:
+			if ((OutputDimensions.X & 0x1) != 0)
+			{
+				return;
+			}
+
+			BufferBytesPerPixel = 2;
+			BufferDimensions = FIntPoint(OutputDimensions.X / 2, OutputDimensions.Y);
+			BufferFormat = PF_B8G8R8A8;
+			RequiresConversion = true;
+			break;
+
+		case EMediaTextureSinkFormat::CharAYUV:
+		case EMediaTextureSinkFormat::CharBGRA:
+			BufferBytesPerPixel = 4;
+			BufferDimensions = OutputDimensions;
+			BufferFormat = PF_B8G8R8A8;
+			RequiresConversion = false;
+			break;
+
+		case EMediaTextureSinkFormat::FloatRGB:
+			BufferBytesPerPixel = 12;
+			BufferDimensions = OutputDimensions;
+			BufferFormat = PF_FloatRGB;
+			RequiresConversion = false;
+			break;
+
+		case EMediaTextureSinkFormat::FloatRGBA:
+			BufferBytesPerPixel = 16;
+			BufferDimensions = OutputDimensions;
+			BufferFormat = PF_FloatRGBA;
+			RequiresConversion = false;
+			break;
+
+		default:
+			return;
+		}
+	}
+
+	// determine how many buffer resources we need
+	int32 StartIndex;
+	
+	if (SinkMode == EMediaTextureSinkMode::Buffered)
+	{
+		StartIndex = 0; // triple buffer
+	}
+	else if (RequiresConversion)
+	{
+		StartIndex = 2; // single buffer (frame before pixel format conversion)
+	}
+	else
+	{
+		StartIndex = 3; // no buffering (directly writing to output or replacing resources)
+	}
+
+	// create buffer resources
+	if (StartIndex < 3)
+	{
+		TripleBuffer.Reset();
+
+		FRHIResourceCreateInfo CreateInfo;
+		uint32 OutStride = 0;
+
+		// only one mip is used & textures are targetable or resolve
+		for (int32 ResourceIndex = StartIndex; ResourceIndex < 3; ++ResourceIndex)
+		{
+			FResource& Resource = BufferResources[ResourceIndex];
+
+			RHICreateTargetableShaderResource2D(
+				BufferDimensions.X,
+				BufferDimensions.Y,
+				BufferFormat,
+				1,
+				TexCreate_Dynamic | TexCreate_NoTiling,
+				TexCreate_RenderTargetable,
+				false,
+				CreateInfo,
+				Resource.RenderTarget,
+				Resource.ShaderResource
+			);
+
+			if (ResourceIndex < 2)
+			{
+				// lock write & temp buffers (in triple buffered mode only)
+				Resource.LockedData = RHILockTexture2D(Resource.RenderTarget.GetReference(), 0, RLM_WriteOnly, OutStride, false);
+			}
+			else
+			{
+				Resource.LockedData = nullptr;
+			}
+		}
+	}
+
+	// create output resource if needed
+	if (RequiresConversion || (SinkMode == EMediaTextureSinkMode::Unbuffered))
+	{
+		// In Unbuffered mode or in Buffered mode (with conversion), this
+		// resource is used as a render target for the output frame.
+
+		FRHIResourceCreateInfo CreateInfo;
+		TRefCountPtr<FRHITexture2D> Texture2DRHI;
+
+		RHICreateTargetableShaderResource2D(
+			OutputDimensions.X,
+			OutputDimensions.Y,
+			RequiresConversion ? PF_B8G8R8A8 : BufferFormat,
+			1,
+			TexCreate_Dynamic | TexCreate_NoTiling | (Owner.SRGB ? TexCreate_SRGB : 0),
+			TexCreate_RenderTargetable,
+			false,
+			CreateInfo,
+			RenderTargetTextureRHI,
+			Texture2DRHI
+		);
+
+		TextureRHI = (FTextureRHIRef&)Texture2DRHI;
+	}
+	else
+	{
+		check(StartIndex < 3);
+
+		// in Buffered mode, we pick the current read buffer as the first resource
+		// to be displayed. In Unbuffered mode (with conversion), the read buffer
+		// is used as the output render target.
+
+		const FResource& Resource = TripleBuffer.Read();
+		
+		RenderTargetTextureRHI = Resource.RenderTarget;
+		TextureRHI = Resource.ShaderResource;
+	}
+
+	// set the initial resource to be displayed
+	RHIUpdateTextureReference(Owner.TextureReference.TextureReferenceRHI, RenderTargetTextureRHI.GetReference());
+
+	// clear the displayed resource
+	FRHICommandListImmediate& CommandList = FRHICommandListExecutor::GetImmediateCommandList();
+	{
+		SetRenderTarget(CommandList, RenderTargetTextureRHI, FTextureRHIRef());
+		CommandList.SetViewport(0, 0, 0.0f, OutputDimensions.X, OutputDimensions.Y, 1.0f);
+		CommandList.Clear(true, BufferClearColor, false, 0.0f, false, 0, FIntRect());
+		CommandList.CopyToResolveTarget(TextureRHI, TextureRHI, true, FResolveParams());
+	}
+
+	State = EState::Initialized;
+	FPlatformMisc::MemoryBarrier();
+}
+
+
+void FMediaTextureResource::ReleaseDynamicRHI()
+{
+	// release output resource if needed
+	if (RequiresConversion || (SinkMode == EMediaTextureSinkMode::Unbuffered))
+	{
+		RenderTargetTextureRHI.SafeRelease();
+		TextureRHI.SafeRelease();
+	}
+
+	// release buffer resources
+	int32 StartIndex;
+
+	if (SinkMode == EMediaTextureSinkMode::Buffered)
+	{
+		StartIndex = 0; // triple buffer
+	}
+	else if (RequiresConversion)
+	{
+		StartIndex = 2; // single buffer (for pixel format conversion)
+	}
+	else
+	{
+		StartIndex = 3; // no buffering (writing directly to output)
+	}
+
+	for (int32 ResourceIndex = StartIndex; ResourceIndex < 3; ++ResourceIndex)
+	{
+		FResource& Resource = BufferResources[ResourceIndex];
+
+		if (Resource.LockedData != nullptr)
+		{
+			RHIUnlockTexture2D(Resource.RenderTarget.GetReference(), 0, false);
+			Resource.LockedData = nullptr;
+		}
+
+		Resource.ShaderResource.SafeRelease();
+		Resource.RenderTarget.SafeRelease();
+	}
+
+	// clean up
+	if (SinkMode == EMediaTextureSinkMode::Buffered)
+	{
+		TripleBuffer.Reset();
+	}
+
+	State = EState::ShutDown;
+	FPlatformMisc::MemoryBarrier();
+}
+
+
+/* FTickerObjectBase interface
+ *****************************************************************************/
+
+bool FMediaTextureResource::Tick(float DeltaTime)
+{
+	ENQUEUE_UNIQUE_RENDER_COMMAND_ONEPARAMETER(TickMediaTextureResource, FMediaTextureResource*, This, this,
+	{
+		TFunction<void()> Task;
+
+		while (This->RenderThreadTasks.Dequeue(Task))
+		{
+			Task();
+		}
+	});
+
+	return true;
+}
+
+
+/* FMediaTextureResource implementation
+ *****************************************************************************/
+
+void FMediaTextureResource::ConvertResource(const FResource& Resource)
+{
+	check(IsInRenderingThread());
+
+	FRHICommandListImmediate& RHICmdList = FRHICommandListExecutor::GetImmediateCommandList();
+
+	SCOPED_DRAW_EVENT(RHICmdList, MediaTextureConvertResource);
+
+	// configure media shaders
+	auto ShaderMap = GetGlobalShaderMap(ERHIFeatureLevel::SM5);
+
+	static FGlobalBoundShaderState BoundShaderState;
+	TShaderMapRef<FMediaShadersVS> VertexShader(ShaderMap);
+
+	switch (SinkFormat)
+	{
+	case EMediaTextureSinkFormat::CharAYUV:
+		{
+			TShaderMapRef<FAYUVConvertPS> ConvertShader(ShaderMap);
+			SetGlobalBoundShaderState(RHICmdList, ERHIFeatureLevel::SM5, BoundShaderState, GMediaVertexDeclaration.VertexDeclarationRHI, *VertexShader, *ConvertShader);
+			ConvertShader->SetParameters(RHICmdList, Resource.ShaderResource);
+		}
+		break;
+
+	case EMediaTextureSinkFormat::CharNV12:
+		{
+			TShaderMapRef<FNV12ConvertPS> ConvertShader(ShaderMap);
+			SetGlobalBoundShaderState(RHICmdList, ERHIFeatureLevel::SM5, BoundShaderState, GMediaVertexDeclaration.VertexDeclarationRHI, *VertexShader, *ConvertShader);
+			ConvertShader->SetParameters(RHICmdList, Resource.ShaderResource);
+		}
+		break;
+
+	case EMediaTextureSinkFormat::CharNV21:
+		{
+			TShaderMapRef<FNV21ConvertPS> ConvertShader(ShaderMap);
+			SetGlobalBoundShaderState(RHICmdList, ERHIFeatureLevel::SM5, BoundShaderState, GMediaVertexDeclaration.VertexDeclarationRHI, *VertexShader, *ConvertShader);
+			ConvertShader->SetParameters(RHICmdList, Resource.ShaderResource);
+		}
+		break;
+
+	case EMediaTextureSinkFormat::CharUYVY:
+		{
+			TShaderMapRef<FUYVYConvertPS> ConvertShader(ShaderMap);
+			SetGlobalBoundShaderState(RHICmdList, ERHIFeatureLevel::SM5, BoundShaderState, GMediaVertexDeclaration.VertexDeclarationRHI, *VertexShader, *ConvertShader);
+			ConvertShader->SetParameters(RHICmdList, Resource.ShaderResource);
+		}
+		break;
+
+	case EMediaTextureSinkFormat::CharYUY2:
+		{
+			TShaderMapRef<FYUY2ConvertPS> ConvertShader(ShaderMap);
+			SetGlobalBoundShaderState(RHICmdList, ERHIFeatureLevel::SM5, BoundShaderState, GMediaVertexDeclaration.VertexDeclarationRHI, *VertexShader, *ConvertShader);
+			ConvertShader->SetParameters(RHICmdList, Resource.ShaderResource);
+		}
+		break;
+
+	case EMediaTextureSinkFormat::CharYVYU:
+		{
+			TShaderMapRef<FYVYUConvertPS> ConvertShader(ShaderMap);
+			SetGlobalBoundShaderState(RHICmdList, ERHIFeatureLevel::SM5, BoundShaderState, GMediaVertexDeclaration.VertexDeclarationRHI, *VertexShader, *ConvertShader);
+			ConvertShader->SetParameters(RHICmdList, Resource.ShaderResource);
+		}
+		break;
+
+	default:
+		return;
+	}
+
+	// draw full size quad into render target
+	FMediaElementVertex Vertices[4];
+	{
+		Vertices[0].Position.Set(-1.0f, 1.0f, 1.0f, 1.0f);
+		Vertices[0].TextureCoordinate.Set(0.0f, 0.0f);
+		Vertices[1].Position.Set(1.0f, 1.0f, 1.0f, 1.0f);
+		Vertices[1].TextureCoordinate.Set(1.0f, 0.0f);
+		Vertices[2].Position.Set(-1.0f, -1.0f, 1.0f, 1.0f);
+		Vertices[2].TextureCoordinate.Set(0.0f, 1.0f);
+		Vertices[3].Position.Set(1.0f, -1.0f, 1.0f, 1.0f);
+		Vertices[3].TextureCoordinate.Set(1.0f, 1.0f);
+	}
+
+	RHICmdList.SetDepthStencilState(TStaticDepthStencilState<false, CF_Always>::GetRHI());
+	RHICmdList.SetRasterizerState(TStaticRasterizerState<>::GetRHI());
+	RHICmdList.SetBlendState(TStaticBlendStateWriteMask<CW_RGBA, CW_NONE, CW_NONE, CW_NONE, CW_NONE, CW_NONE, CW_NONE, CW_NONE>::GetRHI());
+	RHICmdList.SetViewport(0, 0, 0.0f, OutputDimensions.X, OutputDimensions.Y, 1.0f);
+
+	FRHITexture* RenderTarget = RenderTargetTextureRHI.GetReference();
+	SetRenderTargets(RHICmdList, 1, &RenderTarget, nullptr, ESimpleRenderTargetMode::EExistingColorAndDepth, FExclusiveDepthStencil::DepthNop_StencilNop);
+	DrawPrimitiveUP(RHICmdList, PT_TriangleStrip, 2, Vertices, sizeof(Vertices[0]));
+
+	RHICmdList.CopyToResolveTarget(RenderTargetTextureRHI, RenderTargetTextureRHI, true, FResolveParams());
+}
+
+
+void FMediaTextureResource::DisplayResource(const FResource& Resource)
+{
+	check(IsInRenderingThread());
+
+	RenderTargetTextureRHI = Resource.RenderTarget;
+	TextureRHI = Resource.ShaderResource;
+
+	RHIUpdateTextureReference(Owner.TextureReference.TextureReferenceRHI, RenderTargetTextureRHI.GetReference());
+}
+
+
+void FMediaTextureResource::InitializeResource(FIntPoint Dimensions, EMediaTextureSinkFormat Format, EMediaTextureSinkMode Mode)
+{
+	check(IsInRenderingThread());
+
+	// use existing buffers if possible
+	if ((Dimensions != BufferDimensions) || (Format != SinkFormat) || (Mode != SinkMode))
+	{
+		ReleaseDynamicRHI();
+		OutputDimensions = Dimensions;
+		SinkFormat = Format;
+		SinkMode = Mode;
+		InitDynamicRHI();
+	}
+	else
+	{
+		State = EState::Initialized;
+	}
+}
+
+
+void FMediaTextureResource::SetResource(const FResource& Resource)
+{
+	check(IsInRenderingThread());
+
+	if (!Resource.RenderTarget.IsValid() || !Resource.ShaderResource.IsValid())
+	{
+		ReleaseDynamicRHI();
+		InitDynamicRHI();
+	}
+	else if (RequiresConversion)
+	{
+		ConvertResource(Resource);
+	}
+	else
+	{
+		DisplayResource(Resource);
+	}
+}
+
+
+void FMediaTextureResource::SwapResource()
+{
+	if (SinkMode == EMediaTextureSinkMode::Buffered)
+	{
+		if (TripleBuffer.IsDirty())
+		{
+			// lock & swap out old read buffer
+			FResource& OldResource = TripleBuffer.Read();
+			uint32 OutStride = 0;
+			OldResource.LockedData = RHILockTexture2D(OldResource.RenderTarget.GetReference(), 0, RLM_WriteOnly, OutStride, false /*bLockWithinMiptail*/, false /*bFlushRHIThread*/);
+			TripleBuffer.SwapReadBuffers();
+
+			// unlock & assign new read buffer
+			FResource& NewResource = TripleBuffer.Read();
+			RHIUnlockTexture2D(NewResource.RenderTarget.GetReference(), 0, false /*bLockWithinMiptail*/);
+			NewResource.LockedData = nullptr;
+
+			if (RequiresConversion)
+			{
+				ConvertResource(NewResource);
+			}
+			else
+			{
+				DisplayResource(NewResource);
+			}
+		}
+	}
+	else if (RequiresConversion)
+	{
+		ConvertResource(BufferResources[2]);
+	}
+}
diff --git a/Engine/Source/Runtime/MediaAssets/Private/Misc/MediaTextureResource.h b/Engine/Source/Runtime/MediaAssets/Private/Misc/MediaTextureResource.h
new file mode 100644
--- /dev/null
+++ b/Engine/Source/Runtime/MediaAssets/Private/Misc/MediaTextureResource.h
@@ -0,0 +1,236 @@
+// Copyright 1998-2016 Epic Games, Inc. All Rights Reserved.
+
+#pragma once
+
+#include "TextureResource.h"
+#include "TickableObjectRenderThread.h"
+#include "UnrealClient.h"
+
+
+class UMediaTexture;
+
+
+/**
+ * Texture resource type for media textures.
+ */
+class FMediaTextureResource
+	: public FRenderTarget
+	, public FTextureResource
+	, public FTickerObjectBase
+{
+	struct FResource
+	{
+		void* LockedData;
+		TRefCountPtr<FRHITexture2D> RenderTarget;
+		TRefCountPtr<FRHITexture2D> ShaderResource;
+
+		FResource() : LockedData(nullptr) { }
+	};
+
+public:
+
+	/** 
+	 * Creates and initializes a new instance.
+	 *
+	 * @param InOwner The Movie texture object to create a resource for (must not be nullptr).
+	 * @param InClearColor The clear color to use.
+	 * @param InOutputDimensions
+	 * @param InSinkFormat
+	 * @param InSinkMode
+	 */
+	FMediaTextureResource(UMediaTexture& InOwner, const FLinearColor& InClearColor, const FIntPoint& InOutputDimensions, EMediaTextureSinkFormat InSinkFormat, EMediaTextureSinkMode InSinkMode);
+
+	/** Virtual destructor. */
+	virtual ~FMediaTextureResource();
+
+public:
+
+	/**
+	 * Acquire a pointer to locked texture memory buffer to render to.
+	 *
+	 * @return Pointer to buffer on success, nullptr otherwise.
+	 */
+	void* AcquireBuffer();
+
+	/** Display the latest texture buffer. */
+	void DisplayBuffer();
+
+	/**
+	 * Get the memory size of this resource.
+	 *
+	 * @return Resource size (in bytes).
+	 */
+	SIZE_T GetResourceSize() const;
+
+	/**
+	 * Get the render target texture to render to.
+	 *
+	 * @return The texture, or nullptr if called in Buffered mode or not on rendering thread.
+	 */
+	FRHITexture* GetTexture();
+
+	/**
+	 * Initialize the render target buffer(s).
+	 *
+	 * @param Dimensions Width and height of the texture (in pixels).
+	 * @param Format The pixel format of the sink's render target texture.
+	 * @param Mode The mode to operate the sink in (buffered vs. unbuffered).
+	 */
+	void InitializeBuffer(FIntPoint Dimensions, EMediaTextureSinkFormat Format, EMediaTextureSinkMode Mode);
+
+	/** Release a previously acquired texture buffer. */
+	void ReleaseBuffer();
+
+	/** Shut down the render target buffers. */
+	void ShutdownBuffer();
+
+	/**
+	 * Update the render target texture with new frame data.
+	 *
+	 * @param Data The frame data to copy into the render target.
+	 * @param Pitch Number of bytes per pixel row (0 = default).
+	 */
+	void UpdateBuffer(const uint8* Data, uint32 Pitch);
+
+	/**
+	 * Update the underlying resource textures.
+	 *
+	 * @param RenderTarget The new render target texture.
+	 * @param ShaderResource The new shader resource texture.
+	 */
+	void UpdateTextures(FRHITexture* RenderTarget, FRHITexture* ShaderResource);
+
+public:
+
+	//~ FRenderTarget interface
+
+	virtual FIntPoint GetSizeXY() const override;
+
+public:
+
+	//~ FTextureResource interface
+
+	virtual FString GetFriendlyName() const override;
+	virtual uint32 GetSizeX() const override;
+	virtual uint32 GetSizeY() const override;
+	virtual void InitDynamicRHI() override;
+	virtual void ReleaseDynamicRHI() override;
+
+public:
+
+	//~ FTickerObjectBase interface
+
+	virtual bool Tick(float DeltaTime) override;
+
+protected:
+
+	/**
+	 * Helper function to convert the given resource to the target resource's pixel format.
+	 *
+	 * @param Resource The texture resource to convert.
+	 * @see DisplayResource
+	 */
+	void ConvertResource(const FResource& Resource);
+
+	/**
+	 * Helper function to display the given resource in the media texture.
+	 *
+	 * @param Resource The resource to display.
+	 * @see ConvertResource
+	 */
+	void DisplayResource(const FResource& Resource);
+
+	/**
+	 * Initialize this resource.
+	 *
+	 * @param Dimensions The new texture dimensions.
+	 * @param Format The new texture format.
+	 * @param Mode The new sink mode.
+	 */
+	void InitializeResource(FIntPoint Dimensions, EMediaTextureSinkFormat Format, EMediaTextureSinkMode Mode);
+
+	/**
+	 * Helper function to make the given resource the output.
+	 *
+	 * @param Resource The resource to set.
+	 */
+	void SetResource(const FResource& Resource);
+
+	/** Swap the render target resource for display. */
+	void SwapResource();
+
+private:
+
+	/** The media texture that owns this resource. */
+	UMediaTexture& Owner;
+
+	/** Triple-buffer for texture resources. */
+	TTripleBuffer<FResource> TripleBuffer;
+
+private:
+
+	//~ The following fields are owned by the render thread
+
+	/** Number of bytes per pixel in buffer resources. */
+	uint8 BufferBytesPerPixel;
+
+	/** The clear color to use. */
+	FLinearColor BufferClearColor;
+
+	/**
+	 * Width and height of the buffer resources (in pixels).
+	 *
+	 * This may be less than OutputDimensions, because the buffers can
+	 * have formats with multiple pixels packed into a single RGBA tuple.
+	 */
+	FIntPoint BufferDimensions;
+
+	/**
+	 * Texture resources for buffered mode or pixel conversions.
+	 *
+	 * In Buffered mode, all three resources are used for triple buffering.
+	 * If no pixel format conversion is required, the triple buffer's Read
+	 * buffer is used as the output resource. Otherwise the Read buffer is
+	 * converted into a separate output resource.
+	 *
+	 * In Unbuffered mode when pixel format conversion is required, only the
+	 * third buffer (index = 2) is used. Otherwise data is written directly
+	 * to a separate output resource.
+	 */
+	FResource BufferResources[3];
+
+	/** Width and height of the output resource (in pixels). */
+	FIntPoint OutputDimensions;
+
+	/**
+	 * Asynchronous tasks for the render thread.
+	 *
+	 * We can't reliably enqueue render commands using the existing macros
+	 * from the decoder thread, which runs neither on the game nor on the
+	 * render thread, so we maintain our own queue of render commands here.
+	 */
+	TQueue<TFunction<void()>> RenderThreadTasks;
+
+	/** Whether the current media sink format requires pixel format conversion. */
+	bool RequiresConversion;
+
+	/** The render target's pixel format. */
+	EMediaTextureSinkFormat SinkFormat;
+
+	/** The mode that this sink is currently operating in. */
+	EMediaTextureSinkMode SinkMode;
+
+private:
+
+	//~ The following fields are owned by the decoder thread
+
+	/** The resource's current state. */
+	enum class EState
+	{
+		Initializing,
+		Initialized,
+		ShuttingDown,
+		ShutDown
+	}
+	State;
+};
diff --git a/Engine/Source/Runtime/MediaAssets/Public/FileMediaSource.h b/Engine/Source/Runtime/MediaAssets/Public/FileMediaSource.h
new file mode 100644
--- /dev/null
+++ b/Engine/Source/Runtime/MediaAssets/Public/FileMediaSource.h
@@ -0,0 +1,42 @@
+// Copyright 1998-2016 Epic Games, Inc. All Rights Reserved.
+
+#pragma once
+
+#include "MediaSource.h"
+#include "FileMediaSource.generated.h"
+
+
+UCLASS(BlueprintType, MinimalAPI)
+class UFileMediaSource
+	: public UMediaSource
+{
+	GENERATED_BODY()
+
+public:
+
+	/** The path to the media file to be played. */
+	UPROPERTY(EditAnywhere, Category=File, AssetRegistrySearchable)
+	FString FilePath;
+
+	/** Load entire media file into memory and play from there (if possible). */
+	UPROPERTY(EditAnywhere, Category=File)
+	bool PrecacheFile;
+
+public:
+
+	//~ IMediaOptions interface
+
+	virtual bool GetMediaOption(const FName& Key, bool DefaultValue) const override;
+
+public:
+
+	//~ UMediaSource interface
+
+	virtual FString GetUrl() const override;
+	virtual bool Validate() const override;
+
+protected:
+
+	/** Get the full path to the file. */
+	FString GetFullPath() const;
+};
