Copying //UE4/Dev-Sequencer to //UE4/Dev-Main (Source: //UE4/Dev-Sequencer @ 3057646)

diff --git a/Engine/Source/Runtime/MediaAssets/Public/MediaPlayer.h b/Engine/Source/Runtime/MediaAssets/Public/MediaPlayer.h
--- a/Engine/Source/Runtime/MediaAssets/Public/MediaPlayer.h
+++ b/Engine/Source/Runtime/MediaAssets/Public/MediaPlayer.h
@@ -5,6 +5,12 @@
 #include "MediaPlayer.generated.h"
 
 
+class FMediaCaptionSink;
+class UMediaPlaylist;
+class UMediaSource;
+class UMediaSoundWave;
+class UMediaTexture;
+class IMediaOptions;
 class IMediaPlayer;
 enum class EMediaEvent;
 
@@ -20,10 +26,33 @@ DECLARE_DYNAMIC_MULTICAST_DELEGATE_OneParam(FOnMediaPlayerMediaOpenFailed, FStri
 
 
 /**
- * Implements a media player asset that can play movies and other media.
- *
- * This class is represents a media URL along with a corresponding media player
- * for exposing media playback functionality to the Engine and to Blueprints.
+ * Media track types.
+ */
+UENUM(BlueprintType)
+enum class EMediaPlayerTrack : uint8
+{
+	/** Audio tracks. */
+	Audio,
+
+	/** Binary data tracks. */
+	Binary,
+
+	/** Caption (subtitle) tracks) */
+	Caption,
+
+	/** Still image tracks. */
+	Image,
+
+	/** Script tracks. */
+	Script,
+
+	/** Video tracks. */
+	Video
+};
+
+
+/**
+ * Implements a media player asset that can play movies and other media sources.
  */
 UCLASS(BlueprintType, hidecategories=(Object))
 class MEDIAASSETS_API UMediaPlayer
@@ -45,24 +74,23 @@ public:
 	bool CanPause() const;
 
 	/**
-	 * Checks whether media playback can be started right now.
+	 * Close the currently open media, if any.
 	 *
-	 * @return true if playback can be started, false otherwise.
-	 * @see CanPause, Play
+	 * @see OnMediaClosed, OpenPlaylist, OpenPlaylistIndex, OpenSource, OpenUrl, Pause, Play
 	 */
 	UFUNCTION(BlueprintCallable, Category="Media|MediaPlayer")
-	bool CanPlay() const;
+	void Close();
 
 	/**
-	 * Close the currently open media, if any.
+	 * Get the current caption/subtitle text, if any.
 	 *
-	 * @see OnMediaClosed, Open, Pause, Play
+	 * @return Caption text.
 	 */
 	UFUNCTION(BlueprintCallable, Category="Media|MediaPlayer")
-	void Close();
+	FText GetCaptionText() const;
 
 	/**
-	 * Gets the media's duration.
+	 * Get the media's duration.
 	 *
 	 * @return A time span representing the duration.
 	 * @see GetTime, Seek
@@ -71,7 +99,35 @@ public:
 	FTimespan GetDuration() const;
 
 	/**
-	 * Gets the media's current playback rate.
+	 * Get the supported forward playback rates.
+	 *
+	 * @param Unthinned Whether the rates are for unthinned playback (default = true).
+	 * @return The range of supported rates.
+	 * @see GetReverseRates
+	 */
+	UFUNCTION(BlueprintCallable, Category="Media|MediaPlayer")
+	FFloatRange GetForwardRates(bool Unthinned = true);
+
+	/**
+	 * Get the number of tracks of the given type.
+	 *
+	 * @param TrackType The type of media tracks.
+	 * @return Number of tracks.
+	 * @see GetSelectedTrack, SelectTrack
+	 */
+	UFUNCTION(BlueprintCallable, Category="Media|MediaPlayer")
+	int32 GetNumTracks(EMediaPlayerTrack TrackType) const;
+
+	/**
+	 * Get the name of the current native media player.
+	 *
+	 * @return Player name.
+	 */
+	UFUNCTION(BlueprintCallable, Category="Media|MediaPlayer")
+	FName GetPlayerName() const;
+
+	/**
+	 * Get the media's current playback rate.
 	 *
 	 * @return The playback rate.
 	 * @see SetRate, SupportsRate
@@ -80,7 +136,27 @@ public:
 	float GetRate() const;
 
 	/**
-	 * Gets the media's current playback time.
+	 * Get the supported reverse playback rates.
+	 *
+	 * @param Unthinned Whether the rates are for unthinned playback (default = true).
+	 * @return The range of supported rates.
+	 * @see GetForwardRates
+	 */
+	UFUNCTION(BlueprintCallable, Category="Media|MediaPlayer")
+	FFloatRange GetReverseRates(bool Unthinned = true);
+
+	/**
+	 * Get the index of the currently selected track of the given type.
+	 *
+	 * @param TrackType The type of track to get.
+	 * @return The index of the selected track, or INDEX_NONE if no track is active.
+	 * @see GetNumTracks, SelectTrack
+	 */
+	UFUNCTION(BlueprintCallable, Category="Media|MediaPlayer")
+	int32 GetSelectedTrack(EMediaPlayerTrack TrackType) const;
+
+	/**
+	 * Get the media's current playback time.
 	 *
 	 * @return Playback time.
 	 * @see GetDuration, Seek
@@ -89,13 +165,38 @@ public:
 	FTimespan GetTime() const;
 
 	/**
-	 * Gets the URL of the currently loaded media, if any.
+	 * Get the human readable name of the specified track.
+	 *
+	 * @param TrackType The type of track.
+	 * @param TrackIndex The index of the track.
+	 * @return Display name.
+	 * @see GetNumTracks, GetTrackLanguage
+	 */
+	UFUNCTION(BlueprintCallable, Category = "Media|MediaPlayer")
+	FText GetTrackDisplayName(EMediaPlayerTrack TrackType, int32 TrackIndex) const;
+
+	/**
+	 * Get the language tag of the specified track.
+	 *
+	 * @param TrackType The type of track.
+	 * @param TrackIndex The index of the track.
+	 * @return Language tag, i.e. "en-US" for English, or "und" for undefined.
+	 * @see GetNumTracks, GetTrackDisplayName
+	 */
+	UFUNCTION(BlueprintCallable, Category="Media|MediaPlayer")
+	FString GetTrackLanguage(EMediaPlayerTrack TrackType, int32 TrackIndex) const;
+
+	/**
+	 * Get the URL of the currently loaded media, if any.
 	 *
 	 * @return Media URL, or empty string if no media was loaded.
 	 * @see OpenUrl
 	 */
 	UFUNCTION(BlueprintCallable, Category="Media|MediaPlayer")
-	const FString& GetUrl() const;
+	const FString& GetUrl() const
+	{
+		return CurrentUrl;
+	}
 
 	/**
 	 * Checks whether playback is looping.
@@ -124,9 +225,21 @@ public:
 	UFUNCTION(BlueprintCallable, Category="Media|MediaPlayer")
 	bool IsPlaying() const;
 
+	/**
+	 * Checks whether the media is currently opening or buffering.
+	 *
+	 * @return true if playback is being prepared, false otherwise.
+	 * @see CanPlay, IsPaused, IsReady, Play
+	 */
+	UFUNCTION(BlueprintCallable, Category="Media|MediaPlayer")
+	bool IsPreparing() const;
+
 	/**
 	 * Checks whether media is ready for playback.
 	 *
+	 * A player is ready for playback if it has a media source opened that
+	 * finished preparing and is not in an error state.
+	 *
 	 * @return true if media is ready, false otherwise.
 	 * @see IsPaused, IsPlaying, Stop
 	 */
@@ -134,14 +247,70 @@ public:
 	bool IsReady() const;
 
 	/**
-	 * Opens the specified media URL.
+	 * Open the next item in the current play list.
+	 *
+	 * The player will start playing the new media source if it was playing
+	 * something previously, otherwise it will only open the media source.
 	 *
-	 * @param NewUrl The URL to open.
 	 * @return true on success, false otherwise.
-	 * @see GetUrl, Close
+	 * @see Close, OpenUrl, OpenSource, Play, Previous, SetPlaylist
+	 */
+	UFUNCTION(BlueprintCallable, Category="Media|MediaPlayer")
+	bool Next();
+
+	/**
+	 * Open the first media source in the specified play list.
+	 *
+	 * @param InPlaylist The play list to open.
+	 * @return true if the source will be opened, false otherwise.
+	 * @see Close, OpenPlaylistIndex, OpenSource, OpenUrl, Reopen
+	 */
+	UFUNCTION(BlueprintCallable, Category="Media|MediaPlayer")
+	bool OpenPlaylist(UMediaPlaylist* InPlaylist)
+	{
+		return OpenPlaylistIndex(InPlaylist, 0);
+	}
+
+	/**
+	 * Open a particular media source in the specified play list.
+	 *
+	 * @param InPlaylist The play list to open.
+	 * @param Index The index of the source to open.
+	 * @return true if the source will be opened, false otherwise.
+	 * @see Close, OpenPlaylist, OpenSource, OpenUrl, Reopen
+	 */
+	UFUNCTION(BlueprintCallable, Category="Media|MediaPlayer")
+	bool OpenPlaylistIndex(UMediaPlaylist* InPlaylist, int32 Index);
+
+	/**
+	 * Open the specified media source.
+	 *
+	 * A return value of true indicates that the player will attempt to open
+	 * the media, but it may fail to do so later for other reasons, i.e. if
+	 * a connection to the media server timed out. Use the OnMediaOpened and
+	 * OnMediaOpenFailed delegates to detect if and when the media is ready!
+	 *
+	 * @param MediaSource The media source to open.
+	 * @return true if the source will be opened, false otherwise.
+	 * @see Close, OpenPlaylist, OpenPlaylistIndex, OpenUrl, Reopen
 	 */
 	UFUNCTION(BlueprintCallable, Category="Media|MediaPlayer")
-	bool OpenUrl(const FString& NewUrl);
+	bool OpenSource(UMediaSource* MediaSource);
+
+	/**
+	 * Opens the specified media URL.
+	 *
+	 * A return value of true indicates that the player will attempt to open
+	 * the media, but it may fail to do so later for other reasons, i.e. if
+	 * a connection to the media server timed out. Use the OnMediaOpened and
+	 * OnMediaOpenFailed delegates to detect if and when the media is ready!
+	 *
+	 * @param Url The URL to open.
+	 * @return true if the URL will be opened, false otherwise.
+	 * @see GetUrl, Close, OpenPlaylist, OpenPlaylistIndex, OpenSource, Reopen
+	 */
+	UFUNCTION(BlueprintCallable, Category="Media|MediaPlayer")
+	bool OpenUrl(const FString& Url);
 
 	/**
 	 * Pauses media playback.
@@ -149,7 +318,7 @@ public:
 	 * This is the same as setting the playback rate to 0.0.
 	 *
 	 * @return true if playback is being paused, false otherwise.
-	 * @see CanPause, Close, Play, Rewind, Seek, SetRate
+	 * @see CanPause, Close, Next, Play, Previous, Rewind, Seek
 	 */
 	UFUNCTION(BlueprintCallable, Category="Media|MediaPlayer")
 	bool Pause();
@@ -160,18 +329,39 @@ public:
 	 * This is the same as setting the playback rate to 1.0.
 	 *
 	 * @return true if playback is starting, false otherwise.
-	 * @see CanPlay, Close, Pause, Rewind, Seek, SetRate
+	 * @see CanPlay, GetRate, Next, Pause, Previous, SetRate
 	 */
 	UFUNCTION(BlueprintCallable, Category="Media|MediaPlayer")
 	bool Play();
 
+	/**
+	 * Open the previous item in the current play list.
+	 *
+	 * The player will start playing the new media source if it was playing
+	 * something previously, otherwise it will only open the media source.
+	 *
+	 * @return true on success, false otherwise.
+	 * @see Close, Next, OpenUrl, OpenSource, Play, SetPlaylist
+	 */
+	UFUNCTION(BlueprintCallable, Category="Media|MediaPlayer")
+	bool Previous();
+
+	/**
+	 * Reopens the currently opened media or play list.
+	 *
+	 * @return true if the media will be opened, false otherwise.
+	 * @see Close, Open, OpenPlaylist, OpenPlaylistIndex, OpenSource, OpenUrl
+	 */
+	UFUNCTION(BlueprintCallable, Category="Media|MediaPlayer")
+	bool Reopen();
+
 	/**
 	 * Rewinds the media to the beginning.
 	 *
 	 * This is the same as seeking to zero time.
 	 *
 	 * @return true if rewinding, false otherwise.
-	 * @see GetTime, Close, Pause, Play, Seek
+	 * @see GetTime, Seek
 	 */
 	UFUNCTION(BlueprintCallable, Category="Media|MediaPlayer")
 	bool Rewind();
@@ -186,6 +376,28 @@ public:
 	UFUNCTION(BlueprintCallable, Category="Media|MediaPlayer")
 	bool Seek(const FTimespan& InTime);
 
+	/**
+	 * Select the active track of the given type.
+	 *
+	 * Only one track of a given type can be active at any time.
+	 *
+	 * @param TrackType The type of track to select.
+	 * @param TrackIndex The index of the track to select, or INDEX_NONE to deselect.
+	 * @return true if the track was selected, false otherwise.
+	 * @see GetNumTracks, GetSelectedTrack
+	 */
+	UFUNCTION(BlueprintCallable, Category="Media|MediaPlayer")
+	bool SelectTrack(EMediaPlayerTrack TrackType, int32 TrackIndex);
+
+	/**
+	 * Assign the given texture to the player's image sink.
+	 *
+	 * @param NewTexture The texture to set.
+	 * @see SetSoundWave, SetVideoTexture
+	 */
+	UFUNCTION(BlueprintCallable, Category="Media|MediaPlayer")
+	void SetImageTexture(UMediaTexture* NewTexture);
+
 	/**
 	 * Enables or disables playback looping.
 	 *
@@ -206,6 +418,24 @@ public:
 	UFUNCTION(BlueprintCallable, Category="Media|MediaPlayer")
 	bool SetRate(float Rate);
 
+	/**
+	 * Assign the given sound wave to the player's audio sink.
+	 *
+	 * @param NewSoundWave The sound wave to set.
+	 * @see SetImageTexture, SetVideoTexture
+	 */
+	UFUNCTION(BlueprintCallable, Category="Media|MediaPlayer")
+	void SetSoundWave(UMediaSoundWave* NewSoundWave);
+
+	/**
+	 * Assign the given texture to the player's video sink.
+	 *
+	 * @param NewTexture The texture to set.
+	 * @see SetImageTexture, SetSoundWave
+	 */
+	UFUNCTION(BlueprintCallable, Category="Media|MediaPlayer")
+	void SetVideoTexture(UMediaTexture* NewTexture);
+
 	/**
 	 * Checks whether the specified playback rate is supported.
 	 *
@@ -236,103 +466,223 @@ public:
 
 public:
 
-	/** Holds a delegate that is invoked when playback has reached the end of the media. */
+	/** A delegate that is invoked when playback has reached the end of the media. */
 	UPROPERTY(BlueprintAssignable, Category="Media|MediaPlayer")
 	FOnMediaPlayerMediaEvent OnEndReached;
 
-	/** Holds a delegate that is invoked when a media source has been closed. */
+	/** A delegate that is invoked when a media source has been closed. */
 	UPROPERTY(BlueprintAssignable, Category="Media|MediaPlayer")
 	FOnMediaPlayerMediaEvent OnMediaClosed;
 
-	/** Holds a delegate that is invoked when a media source has been opened. */
+	/**
+	 * A delegate that is invoked when a media source has been opened.
+	 *
+	 * Depending on whether the underlying player implementation opens the media
+	 * synchronously or asynchronously, this event may be executed before or
+	 * after the call top OpenSource / OpenUrl returns.
+	 */
 	UPROPERTY(BlueprintAssignable, Category="Media|MediaPlayer")
 	FOnMediaPlayerMediaOpened OnMediaOpened;
 
-	/** Holds a delegate that is invoked when a media source has failed to open. */
+	/**
+	 * A delegate that is invoked when a media source has failed to open.
+	 *
+	 * This delegate is only executed if OpenSource / OpenUrl returned true and
+	 * the media failed to open asynchronously later. It is not executed if
+	 * OpenSource / OpenUrl returned false, indicating an immediate failure.
+	 */
 	UPROPERTY(BlueprintAssignable, Category="Media|MediaPlayer")
 	FOnMediaPlayerMediaOpenFailed OnMediaOpenFailed;
 
-	/** Holds a delegate that is invoked when media playback has been resumed. */
+	/** A delegate that is invoked when media playback has been resumed. */
 	UPROPERTY(BlueprintAssignable, Category="Media|MediaPlayer")
 	FOnMediaPlayerMediaEvent OnPlaybackResumed;
 
-	/** Holds a delegate that is invoked when media playback has been suspended. */
+	/** A delegate that is invoked when media playback has been suspended. */
 	UPROPERTY(BlueprintAssignable, Category="Media|MediaPlayer")
 	FOnMediaPlayerMediaEvent OnPlaybackSuspended;
 
 public:
 
-	/** Gets an event delegate that is invoked when media has been opened or closed. */
-	DECLARE_EVENT(UMediaPlayer, FOnMediaChanged)
-	FOnMediaChanged& OnMediaChanged()
+	/** Get an event delegate that is invoked when a media event occurred. */
+	DECLARE_EVENT_OneParam(UMediaPlayer, FOnMediaEvent, EMediaEvent /*Event*/)
+	FOnMediaEvent& OnMediaEvent()
 	{
-		return MediaChangedEvent;
+		return MediaEvent;
 	}
 
-	/** Gets an event delegate that is invoked when media tracks have changed. */
-	DECLARE_EVENT(UMediaPlayer, FOnTracksChanged)
-	FOnTracksChanged& OnTracksChanged()
+public:
+
+	/**
+	 * Get the media texture that receives the image output.
+	 *
+	 * @return Media texture asset.
+	 * @see GetPlayer, GetSoundWave, GetVideoTexture
+	 */
+	UMediaTexture* GetImageTexture() const
 	{
-		return TracksChangedEvent;
+		return ImageTexture;
 	}
 
-public:
-
 	/**
-	 * Gets the low-level player associated with this object.
+	 * Get the low-level player associated with this object.
 	 *
 	 * @return The player, or nullptr if no player was created.
+	 * @see GetPlaylist, GetPlaylistIndex
 	 */
 	TSharedPtr<IMediaPlayer> GetPlayer() const
 	{
 		return Player;
 	}
 
+	/**
+	 * Get the currently active play list.
+	 *
+	 * @return The play list, if any.
+	 * @see GetPlayer, GetPlaylistIndex
+	 */
+	UMediaPlaylist* GetPlaylist() const
+	{
+		return Playlist;
+	}
+
+	/**
+	 * Get the current play list index.
+	 *
+	 * @return Play list index.
+	 * @see GetPlayer, GetPlaylist
+	 */
+	int32 GetPlaylistIndex() const
+	{
+		return PlaylistIndex;
+	}
+
+	/**
+	 * Get the sound wave that receives the audio output.
+	 *
+	 * @return Sound wave asset.
+	 * @see GetImageTexture, GetPlayer, GetVideoTexture
+	 */
+	UMediaSoundWave* GetSoundWave() const
+	{
+		return SoundWave;
+	}
+
+	/**
+	 * Get the media texture that receives the video output.
+	 *
+	 * @return Media texture asset.
+	 * @see GetImageTexture, GetPlayer, GetSoundWave
+	 */
+	UMediaTexture* GetVideoTexture() const
+	{
+		return VideoTexture;
+	}
+
+#if WITH_EDITOR
+	/** Called when PIE has been paused. */
+	void PausePIE();
+
+	/** Called when PIE has been resumed. */
+	void ResumePIE();
+#endif
+
 public:
 
-	// UObject overrides.
+	//~ UObject interface
 
 	virtual void BeginDestroy() override;
 	virtual FString GetDesc() override;
-	virtual void PostLoad() override;
+	virtual void PostLoad() override; 
 
 #if WITH_EDITOR
 	virtual void PostEditChangeProperty(FPropertyChangedEvent& PropertyChangedEvent) override;
+	virtual void PreEditChange(UProperty* PropertyAboutToChange) override;
 #endif
 
 protected:
 
-	/** Initializes the media player. */
-	void InitializePlayer();
+	/**
+	 * Open a media source from a URL with optional parameters.
+	 *
+	 * @param Url The URL of the media to open (file name or web address).
+	 * @param Options Optional media parameters.
+	 * @return true if the media is being opened, false otherwise.
+	 */
+	bool Open(const FString& Url, const IMediaOptions& Options);
 
-protected:
+	/** Select the default media tracks. */
+	void SelectDefaultTracks();
+
+public:
+
+	/** Name of the desired native player, if any. */
+	UPROPERTY(transient)
+	FName DesiredPlayerName;
 
-	/** Whether playback should loop when it reaches the end. */
+	/** Automatically start playback after media opened successfully. */
 	UPROPERTY(EditAnywhere, Category=Playback)
-	uint32 Looping:1;
+	bool PlayOnOpen;
+
+	/** Whether playback should shuffle media sources in the play list. */
+	UPROPERTY(BlueprintReadOnly, EditAnywhere, Category=Playback)
+	uint32 Shuffle : 1;
+
+protected:
+
+	/** The media texture to output the image track frames to. */
+	UPROPERTY(BlueprintReadOnly, EditAnywhere, Category=Output, AdvancedDisplay)
+	UMediaTexture* ImageTexture;
+
+	/** Whether the player should loop when media playback reaches the end. */
+	UPROPERTY(BlueprintReadOnly, EditAnywhere, Category=Playback)
+	uint32 Loop:1;
 
-	/** The path or URL to the media file to be played. */
-	UPROPERTY(EditAnywhere, Category=Source)
-	FString URL;
+	/** The play list to use, if any. */
+	UPROPERTY(BlueprintReadOnly, transient, Category=Playback)
+	UMediaPlaylist* Playlist;
+
+	/** The current index of the source in the play list being played. */
+	UPROPERTY(BlueprintReadOnly, Category=Playback)
+	int32 PlaylistIndex;
+
+	/** The media sound wave to output the audio track samples to. */
+	UPROPERTY(BlueprintReadOnly, EditAnywhere, Category=Output)
+	UMediaSoundWave* SoundWave;
+
+	/** The media texture to output the video track frames to. */
+	UPROPERTY(BlueprintReadOnly, EditAnywhere, Category=Output)
+	UMediaTexture* VideoTexture;
 
 private:
 
+	/** Callback for when a media texture is being destroyed. */
+	void HandleMediaSoundWaveBeginDestroy(UMediaSoundWave& DestroyedSoundWave);
+
+	/** Callback for when a media texture is being destroyed. */
+	void HandleMediaTextureBeginDestroy(UMediaTexture& DestroyedMediaTexture);
+
 	/** Callback for when a media event occurred in the player. */
 	void HandlePlayerMediaEvent(EMediaEvent Event);
 
 private:
 
+	/** Caption text sink. */
+	FMediaCaptionSink* CaptionSink;
+
 	/** Holds the URL of the currently loaded media source. */
 	FString CurrentUrl;
 
+	/** An event delegate that is invoked when a media event occurred. */
+	FOnMediaEvent MediaEvent;
+
 	/** Holds the low-level player used to play the media source. */
 	TSharedPtr<IMediaPlayer> Player;
 
-private:
-
-	/** Holds a delegate that is executed when media has been opened or closed. */
-	FOnMediaChanged MediaChangedEvent;
+	/** Name of the current player. */
+	FName PlayerName;
 
-	/** Holds a delegate that is executed when media tracks have changed. */
-	FOnTracksChanged TracksChangedEvent;
+#if WITH_EDITOR
+	bool WasPlayingInPIE;
+#endif
 };
diff --git a/Engine/Source/Runtime/MediaAssets/Public/MediaSoundWave.h b/Engine/Source/Runtime/MediaAssets/Public/MediaSoundWave.h
--- a/Engine/Source/Runtime/MediaAssets/Public/MediaSoundWave.h
+++ b/Engine/Source/Runtime/MediaAssets/Public/MediaSoundWave.h
@@ -2,58 +2,37 @@
 
 #pragma once
 
-#include "MediaSampleQueue.h"
-#include "Sound/SoundWave.h"
+#include "Sound/SoundWaveProcedural.h"
+#include "IMediaAudioSink.h"
 #include "MediaSoundWave.generated.h"
 
 
-class FMediaSampleQueue;
-class IMediaAudioTrack;
-class IMediaPlayer;
 class UMediaPlayer;
 
 
 /**
  * Implements a playable sound asset for audio streams from UMediaPlayer assets.
  */
-UCLASS(hidecategories=(Compression, Sound, SoundWave, Subtitles))
+UCLASS(hidecategories=(Compression, Info, Sound, SoundWave, Subtitles))
 class MEDIAASSETS_API UMediaSoundWave
 	: public USoundWave
+	, public IMediaAudioSink
 {
 	GENERATED_UCLASS_BODY()
 
-	/** The index of the MediaPlayer's audio track to get the wave data from. */
-	UPROPERTY(EditAnywhere, BlueprintReadWrite, Category=MediaPlayer)
-	int32 AudioTrackIndex;
-
-	/** The MediaPlayer asset to stream audio from. */
-	UPROPERTY(EditAnywhere, BlueprintReadWrite, Category=MediaPlayer)
-	UMediaPlayer* MediaPlayer;
-
-public:
-
-	/**
-	 * Sets the MediaPlayer asset to be used for this texture.
-	 *
-	 * @param InMediaPlayer The asset to set.
-	 */
-	UFUNCTION(BlueprintCallable, Category="Media|MediaSound")
-	void SetMediaPlayer(UMediaPlayer* InMediaPlayer);
-
 public:
 
-	/**
-	 * Gets the low-level player associated with the assigned UMediaPlayer asset.
-	 *
-	 * @return The player, or nullptr if no player is available.
-	 */
-	TSharedPtr<IMediaPlayer> GetPlayer() const;
+	DECLARE_EVENT_OneParam(UMediaSoundWave, FOnBeginDestroy, UMediaSoundWave& /*DestroyedSoundWave*/)
+	FOnBeginDestroy& OnBeginDestroy()
+	{
+		return BeginDestroyEvent;
+	}
 
 public:
 
-	// USoundWave overrides
+	//~ USoundWave interface
 
-	virtual int32 GeneratePCMData(uint8* PCMData, const int32 SamplesNeeded) override;
+	virtual int32 GeneratePCMData(uint8* Data, const int32 SamplesRequested) override;
 	virtual FByteBulkData* GetCompressedData(FName Format) override;
 	virtual int32 GetResourceSizeForFormat(FName Format) override;
 	virtual void InitAudioResource(FByteBulkData& CompressedData) override;
@@ -61,38 +40,55 @@ public:
 
 public:
 
-	// UObject overrides
+	//~ UObject interface
 
+	virtual void BeginDestroy() override;
 	virtual void GetAssetRegistryTags(TArray<FAssetRegistryTag>& OutTags) const override;
 	virtual SIZE_T GetResourceSize(EResourceSizeMode::Type Mode) override;
 	virtual void Serialize(FArchive& Ar) override;
-	virtual void PostLoad() override;
-	virtual void BeginDestroy() override;
+
+public:
+
+	//~ IMediaAudioSink interface
+
+	virtual void FlushAudioSink() override;
+	virtual int32 GetAudioSinkChannels() const override;
+	virtual int32 GetAudioSinkSampleRate() const override;
+	virtual bool InitializeAudioSink(uint32 Channels, uint32 SampleRate) override;
+	virtual void PauseAudioSink() override;
+	virtual void PlayAudioSink(const uint8* SampleBuffer, uint32 BufferSize, FTimespan Time) override;
+	virtual void ResumeAudioSink() override;
+	virtual void ShutdownAudioSink() override;
 
 protected:
 
-	/** Initializes the audio track. */
-	void InitializeTrack();
+	//~ Deprecated members
 
-private:
+	DEPRECATED_FORGAME(4.13, "The AudioTrackIndex property is no longer used. Please upgrade your content to Media Framework 2.0")
+	UPROPERTY(BlueprintReadWrite, Category=MediaPlayer)
+	int32 AudioTrackIndex;
 
-	/** Callback for when the UMediaPlayer changed tracks. */
-	void HandleMediaPlayerTracksChanged();
+	DEPRECATED_FORGAME(4.13, "The MediaPlayer property is no longer used. Please upgrade your content to Media Framework 2.0")
+	UPROPERTY(BlueprintReadWrite, Category=MediaPlayer)
+	class UMediaPlayer* MediaPlayer;
 
 private:
 
-	/** The audio sample queue. */
-	TSharedRef<FMediaSampleQueue, ESPMode::ThreadSafe> AudioQueue;
+	/** An event delegate that is invoked when this media texture is being destroyed. */
+	FOnBeginDestroy BeginDestroyEvent;
 
-	/** Holds the selected audio track. */
-	TSharedPtr<IMediaAudioTrack, ESPMode::ThreadSafe> AudioTrack;
+	/** Critical section for synchronizing access to QueuedAudio. */
+	mutable FCriticalSection CriticalSection;
 
-	/** Holds the media player asset currently being used. */
-	UPROPERTY()
-	TWeakObjectPtr<UMediaPlayer> CurrentMediaPlayer;
+	/** Whether the audio is currently paused. */
+	bool Paused;
 
 	/** Holds queued audio samples. */
 	TArray<uint8> QueuedAudio;
 
-	bool bSetupDelegates;
+	/** The sink's number of audio channels. */
+	int32 SinkNumChannels;
+
+	/** The sink's sample rate. */
+	int32 SinkSampleRate;
 };
diff --git a/Engine/Source/Runtime/MediaAssets/Public/MediaSource.h b/Engine/Source/Runtime/MediaAssets/Public/MediaSource.h
new file mode 100644
--- /dev/null
+++ b/Engine/Source/Runtime/MediaAssets/Public/MediaSource.h
@@ -0,0 +1,75 @@
+// Copyright 1998-2016 Epic Games, Inc. All Rights Reserved.
+
+#pragma once
+
+#include "IMediaOptions.h"
+#include "MediaSource.generated.h"
+
+
+/**
+ * Abstract base class for media sources.
+ *
+ * Media sources describe the location and/or settings of media objects that can
+ * be played in a media player, such as a video file on disk, a video stream on
+ * the internet, or a web cam attached to or built into the target device.
+ */
+UCLASS(Abstract, BlueprintType, MinimalAPI, hidecategories=(Object))
+class UMediaSource
+	: public UObject
+	, public IMediaOptions
+{
+	GENERATED_BODY()
+
+public:
+
+	/** Name of the default native media player for all platforms (Empty = find one automatically). */
+	UPROPERTY(EditAnywhere, Category=Playback)
+	FString DefaultPlayer;
+
+#if WITH_EDITORONLY_DATA
+
+	/** Override native media player plug-ins per platform (Empty = find one automatically). */
+	UPROPERTY(EditAnywhere, Category=Playback, AdvancedDisplay)
+	TMap<FString, FString> PlatformPlayers;
+
+#endif
+
+public:
+
+	/**
+	 * Get the name of the desired native player for this media source.
+	 *
+	 * @return Native player name, or empty string if not set.
+	 */
+	FString GetDesiredPlayer() const;
+
+public:
+
+	/** Get the media source's URL string (must override in child classes). */
+	virtual FString GetUrl() const PURE_VIRTUAL(UMediaSource::GetUrl, return FString(););
+
+	/** Validate the media source settings (must override in child classes). */
+	virtual bool Validate() const PURE_VIRTUAL(UMediaSource::Validate, return false;);
+
+public:
+
+	//~ UObject interface
+
+	virtual void GetAssetRegistryTags(TArray<FAssetRegistryTag>& OutTags) const override;
+	virtual void Serialize(FArchive& Ar) override;
+
+#if WITH_EDITOR
+	virtual void GetAssetRegistryTagMetadata(TMap<FName, FAssetRegistryTagMetadata>& OutMetadata) const override;
+#endif
+
+public:
+
+	//~ IMediaOptions interface
+
+	virtual bool GetMediaOption(const FName& Key, bool DefaultValue) const override;
+	virtual double GetMediaOption(const FName& Key, double DefaultValue) const override;
+	virtual int64 GetMediaOption(const FName& Key, int64 DefaultValue) const override;
+	virtual FString GetMediaOption(const FName& Key, const FString& DefaultValue) const override;
+	virtual FText GetMediaOption(const FName& Key, const FText& DefaultValue) const override;
+	virtual bool HasMediaOption(const FName& Key) const override;
+};
diff --git a/Engine/Source/Runtime/MediaAssets/Public/MediaTexture.h b/Engine/Source/Runtime/MediaAssets/Public/MediaTexture.h
--- a/Engine/Source/Runtime/MediaAssets/Public/MediaTexture.h
+++ b/Engine/Source/Runtime/MediaAssets/Public/MediaTexture.h
@@ -3,15 +3,11 @@
 #pragma once
 
 #include "Engine/Texture.h"
+#include "IMediaTextureSink.h"
 #include "MediaTexture.generated.h"
 
 
-class FMediaSampleBuffer;
-class IMediaPlayer;
-class IMediaVideoTrack;
-class UMediaPlayer;
-enum EPixelFormat;
-enum TextureAddress;
+class IMediaTextureSink;
 
 
 /**
@@ -20,6 +16,8 @@ enum TextureAddress;
 UCLASS(hidecategories=(Compression, LevelOfDetail, Object, Texture))
 class MEDIAASSETS_API UMediaTexture
 	: public UTexture
+	, public FTickerObjectBase
+	, public IMediaTextureSink
 {
 	GENERATED_UCLASS_BODY()
 
@@ -31,120 +29,93 @@ class MEDIAASSETS_API UMediaTexture
 	UPROPERTY(EditAnywhere, BlueprintReadWrite, Category=MediaTexture, AssetRegistrySearchable)
 	TEnumAsByte<TextureAddress> AddressY;
 
-	/** The color used to clear the texture if no video data is drawn. */
-	UPROPERTY(EditAnywhere, BlueprintReadWrite, Category=MediaTexture)
+	/** The color used to clear the texture if CloseAction is set to Clear (default = black). */
+	UPROPERTY(EditAnywhere, BlueprintReadWrite, Category=MediaTexture, AdvancedDisplay)
 	FLinearColor ClearColor;
 
-	/** The index of the MediaPlayer's video track to render on this texture. */
-	UPROPERTY(EditAnywhere, BlueprintReadWrite, Category=MediaPlayer)
-	int32 VideoTrackIndex;
-
 public:
 
-	/**
-	 * Gets the width and height of the texture (in pixels).
-	 *
-	 * @return Texture dimensions.
-	 */
-	FIntPoint GetDimensions() const
+	DECLARE_EVENT_OneParam(UMediaTexture, FOnBeginDestroy, UMediaTexture& /*DestroyedMediaTexture*/)
+	FOnBeginDestroy& OnBeginDestroy()
 	{
-		return CachedDimensions;
+		return BeginDestroyEvent;
 	}
 
-	/**
-	 * Gets the texture's pixel format.
-	 *
-	 * @return Pixel format (always PF_B8G8R8A8 for all movie textures).
-	 */
-	TEnumAsByte<EPixelFormat> GetFormat() const
-	{
-		return PF_B8G8R8A8;
-	}
+public:
 
-	/**
-	 * Gets the low-level player associated with the assigned UMediaPlayer asset.
-	 *
-	 * @return The player, or nullptr if no player is available.
-	 */
-	TSharedPtr<IMediaPlayer> GetPlayer() const;
-
-	/**
-	 * Gets the currently selected video track, if any.
-	 *
-	 * @return The selected video track, or nullptr if none is selected.
-	 */
-	TSharedPtr<IMediaVideoTrack, ESPMode::ThreadSafe> GetVideoTrack() const
-	{
-		return VideoTrack;
-	}
-	
-	/**
-	 * Sets the media player asset to be used for this texture.
-	 *
-	 * @param InMediaPlayer The asset to set.
-	 */
-	UFUNCTION(BlueprintCallable, Category="Media|MediaTexture")
-	void SetMediaPlayer(UMediaPlayer* InMediaPlayer);
+	//~ FTickerObjectBase interface
+
+	virtual bool Tick(float DeltaTime) override;
 
 public:
 
-	// UTexture overrides.
+	//~ IMediaTextureSink interface
+
+	virtual void* AcquireTextureSinkBuffer() override;
+	virtual void DisplayTextureSinkBuffer(FTimespan Time) override;
+	virtual FIntPoint GetTextureSinkDimensions() const override;
+	virtual EMediaTextureSinkFormat GetTextureSinkFormat() const override;
+	virtual EMediaTextureSinkMode GetTextureSinkMode() const override;
+	virtual FRHITexture* GetTextureSinkTexture() override;
+	virtual bool InitializeTextureSink(FIntPoint Dimensions, EMediaTextureSinkFormat Format, EMediaTextureSinkMode Mode) override;
+	virtual void ReleaseTextureSinkBuffer() override;
+	virtual void ShutdownTextureSink() override;
+	virtual bool SupportsTextureSinkFormat(EMediaTextureSinkFormat Format) const override;
+	virtual void UpdateTextureSinkBuffer(const uint8* Data, uint32 Pitch = 0) override;
+	virtual void UpdateTextureSinkResource(FRHITexture* RenderTarget, FRHITexture* ShaderResource) override;
+
+public:
+
+	//~ UTexture interface.
 
 	virtual FTextureResource* CreateResource() override;
 	virtual EMaterialValueType GetMaterialType() override;
 	virtual float GetSurfaceWidth() const override;
 	virtual float GetSurfaceHeight() const override;
-	virtual void UpdateResource() override;	
+	virtual void UpdateResource() override;
 
 public:
 
-	// UObject overrides.
+	//~ UObject interface.
 
 	virtual void BeginDestroy() override;
-	virtual void FinishDestroy() override;
 	virtual FString GetDesc() override;
 	virtual SIZE_T GetResourceSize(EResourceSizeMode::Type Mode) override;
-	virtual bool IsReadyForFinishDestroy() override;
-	virtual void PostLoad() override;
 
 #if WITH_EDITOR
-	virtual void PreEditChange(UProperty* PropertyAboutToChange) override;
 	virtual void PostEditChangeProperty(FPropertyChangedEvent& PropertyChangedEvent) override;
+	virtual void PreEditChange(UProperty* PropertyAboutToChange) override;
 #endif
 
 protected:
 
-	/** The MediaPlayer asset to stream video from. */
-	UPROPERTY(EditAnywhere, BlueprintReadOnly, Category=MediaPlayer)
-	UMediaPlayer* MediaPlayer;
-
-protected:
-
-	/** Initializes the video track. */
-	void InitializeTrack();
-
-private:
+	//~ Deprecated members
 
-	/** Callback for when the UMediaPlayer changed tracks. */
-	void HandleMediaPlayerTracksChanged();
+	DEPRECATED_FORGAME(4.13, "The MediaPlayer property is no longer used. Please upgrade your content to Media Framework 2.0")
+	UPROPERTY(BlueprintReadWrite, Category=MediaPlayer)
+	class UMediaPlayer* MediaPlayer;
+	
+	DEPRECATED_FORGAME(4.13, "The VideoTrackIndex property is no longer used. Please upgrade your content to Media Framework 2.0")
+	UPROPERTY(BlueprintReadWrite, Category=MediaPlayer)
+	int32 VideoTrackIndex;
 
 private:
 
-	/** The texture's cached width and height (in pixels). */
-	FIntPoint CachedDimensions;
+	/** An event delegate that is invoked when this media texture is being destroyed. */
+	FOnBeginDestroy BeginDestroyEvent;
 
-	/** Holds the UMediaPlayer asset currently being used. */
-	UPROPERTY() 
-	TWeakObjectPtr<UMediaPlayer> CurrentMediaPlayer;
+	/** Critical section for synchronizing access to texture resource object. */
+	mutable FCriticalSection CriticalSection;
 
-	/** Synchronizes access to this object from the render thread. */
-	FRenderCommandFence* ReleasePlayerFence;
+	/** Asynchronous tasks for the game thread. */
+	TQueue<TFunction<void()>> GameThreadTasks;
 
-	/** The video sample buffer. */
-	TSharedRef<FMediaSampleBuffer, ESPMode::ThreadSafe> VideoBuffer;
+	/** Width and height of the texture (in pixels). */
+	FIntPoint SinkDimensions;
 
-	/** Holds the selected video track. */
-	TSharedPtr<IMediaVideoTrack, ESPMode::ThreadSafe> VideoTrack;
+	/** The render target's pixel format. */
+	EMediaTextureSinkFormat SinkFormat;
 
-	bool bDelegatesAdded;
-};
+	/** The mode that this sink is currently operating in. */
+	EMediaTextureSinkMode SinkMode;
+	};
diff --git a/Engine/Source/Runtime/MediaAssets/Public/PlatformMediaSource.h b/Engine/Source/Runtime/MediaAssets/Public/PlatformMediaSource.h
new file mode 100644
--- /dev/null
+++ b/Engine/Source/Runtime/MediaAssets/Public/PlatformMediaSource.h
@@ -0,0 +1,56 @@
+// Copyright 1998-2016 Epic Games, Inc. All Rights Reserved.
+
+#pragma once
+
+#include "MediaSource.h"
+#include "PlatformMediaSource.generated.h"
+
+
+/**
+ * A media source that selects other media sources based on target platform.
+ *
+ * Use this asset to override media sources on a per-platform basis.
+ */
+UCLASS(BlueprintType, MinimalAPI)
+class UPlatformMediaSource
+	: public UMediaSource
+{
+	GENERATED_BODY()
+
+public:
+
+	/**
+	 * Default media source.
+	 *
+	 * This media source will be used if no source was specified for a target platform.
+	 */
+	UPROPERTY(EditAnywhere, Category=Sources)
+	UMediaSource* DefaultSource;
+
+#if WITH_EDITORONLY_DATA
+
+	/** Media sources per platform. */
+	UPROPERTY(EditAnywhere, Category=Sources)
+	TMap<FString, UMediaSource*> MediaSources;
+
+#endif
+
+public:
+
+	//~ UMediaSource interface
+
+	virtual FString GetUrl() const override;
+	virtual void Serialize(FArchive& Ar) override;
+	virtual bool Validate() const override;
+
+public:
+
+	//~ IMediaOptions interface
+
+	virtual bool GetMediaOption(const FName& Key, bool DefaultValue) const override;
+	virtual double GetMediaOption(const FName& Key, double DefaultValue) const override;
+	virtual int64 GetMediaOption(const FName& Key, int64 DefaultValue) const override;
+	virtual FString GetMediaOption(const FName& Key, const FString& DefaultValue) const override;
+	virtual FText GetMediaOption(const FName& Key, const FText& DefaultValue) const override;
+	virtual bool HasMediaOption(const FName& Key) const override;
+};
diff --git a/Engine/Source/Runtime/MediaAssets/Public/StreamMediaSource.h b/Engine/Source/Runtime/MediaAssets/Public/StreamMediaSource.h
new file mode 100644
--- /dev/null
+++ b/Engine/Source/Runtime/MediaAssets/Public/StreamMediaSource.h
@@ -0,0 +1,27 @@
+// Copyright 1998-2016 Epic Games, Inc. All Rights Reserved.
+
+#pragma once
+
+#include "MediaSource.h"
+#include "StreamMediaSource.generated.h"
+
+
+UCLASS(BlueprintType, MinimalAPI)
+class UStreamMediaSource
+	: public UMediaSource
+{
+	GENERATED_BODY()
+
+public:
+
+	/** The URL to the media stream to be played. */
+	UPROPERTY(EditAnywhere, Category=Stream, AssetRegistrySearchable)
+	FString StreamUrl;
+
+public:
+
+	//~ UMediaSource interface
+
+	virtual FString GetUrl() const override;
+	virtual bool Validate() const override;
+};
diff --git a/Engine/Source/Runtime/MovieScene/Public/MovieScene.h b/Engine/Source/Runtime/MovieScene/Public/MovieScene.h
--- a/Engine/Source/Runtime/MovieScene/Public/MovieScene.h
+++ b/Engine/Source/Runtime/MovieScene/Public/MovieScene.h
@@ -397,7 +397,7 @@ public:
 	 * Get the display name of the object with the specified identifier.
 	 *
 	 * @param ObjectId The object identifier.
-	 * @result The object's display name.
+	 * @return The object's display name.
 	 */
 	FText GetObjectDisplayName(const FGuid& ObjectId);
 
@@ -430,7 +430,7 @@ public:
 	 * Get the display name of the object with the specified identifier.
 	 *
 	 * @param ObjectId The object identifier.
-	 * @result The object's display name.
+	 * @return The object's display name.
 	 */
 	void SetObjectDisplayName(const FGuid& ObjectId, const FText& DisplayName);
 
diff --git a/Engine/Source/Runtime/MovieScene/Public/MovieSceneSection.h b/Engine/Source/Runtime/MovieScene/Public/MovieSceneSection.h
--- a/Engine/Source/Runtime/MovieScene/Public/MovieSceneSection.h
+++ b/Engine/Source/Runtime/MovieScene/Public/MovieSceneSection.h
@@ -14,6 +14,7 @@ class UMovieSceneSection
 	: public UObject
 {
 	GENERATED_UCLASS_BODY()
+
 public:
 
 	/**
@@ -215,13 +216,13 @@ public:
 	}
 
 	/**
-	 * Adds a key to a rich curve, finding an existing key to modify or adding a new one
+	 * Adds a key to a rich curve, finding an existing key to modify or adding a new one.
 	 *
-	 * @param InCurve	The curve to add keys to
-	 * @param Time		The time where the key should be added
-	 * @param Value		The value at the given time
-	 * @param KeyParams The keying parameters
-	 * @param bUnwindRotation Unwind rotation
+	 * @param InCurve The curve to add keys to.
+	 * @param Time The time where the key should be added.
+	 * @param Value The value at the given time.
+	 * @param Interpolation The key interpolation to use.
+	 * @param bUnwindRotation Unwind rotation.
 	 */
 	void MOVIESCENE_API AddKeyToCurve(FRichCurve& InCurve, float Time, float Value, EMovieSceneKeyInterpolation Interpolation, const bool bUnwindRotation = false);
 
@@ -237,20 +238,20 @@ public:
 	 * Checks to see if this section overlaps with an array of other sections
 	 * given an optional time and track delta.
 	 *
-	 * @param Sections		Section array to check against
-	 * @param TrackDelta	Optional offset to this section's track index
-	 * @param TimeDelta		Optional offset to this section's time delta
-	 * @return				The first section that overlaps, or null if there is no overlap
+	 * @param Sections Section array to check against.
+	 * @param TrackDelta Optional offset to this section's track index.
+	 * @param TimeDelta Optional offset to this section's time delta.
+	 * @return The first section that overlaps, or null if there is no overlap.
 	 */
 	virtual MOVIESCENE_API const UMovieSceneSection* OverlapsWithSections(const TArray<UMovieSceneSection*>& Sections, int32 TrackDelta = 0, float TimeDelta = 0.f) const;
 	
 	/**
 	 * Places this section at the first valid row at the specified time. Good for placement upon creation.
 	 *
-	 * @param Sections		Sections that we can not overlap with
-	 * @param InStartTime	The new start time
-	 * @param InEndTime		The new end time
-	 * @param bAllowMultipleRows	If false, it will move the section in the time direction to make it fit, rather than the row direction
+	 * @param Sections Sections that we can not overlap with.
+	 * @param InStartTime The new start time.
+	 * @param InEndTime The new end time.
+	 * @param bAllowMultipleRows If false, it will move the section in the time direction to make it fit, rather than the row direction.
 	 */
 	virtual MOVIESCENE_API void InitialPlacement(const TArray<UMovieSceneSection*>& Sections, float InStartTime, float InEndTime, bool bAllowMultipleRows);
 
diff --git a/Engine/Source/Runtime/MovieSceneTracks/Private/MovieSceneTracksModule.cpp b/Engine/Source/Runtime/MovieSceneTracks/Private/MovieSceneTracksModule.cpp
--- a/Engine/Source/Runtime/MovieSceneTracks/Private/MovieSceneTracksModule.cpp
+++ b/Engine/Source/Runtime/MovieSceneTracks/Private/MovieSceneTracksModule.cpp
@@ -1,7 +1,7 @@
 // Copyright 1998-2016 Epic Games, Inc. All Rights Reserved.
 
 #include "MovieSceneTracksPrivatePCH.h"
-#include "ModuleManager.h"
+#include "IMovieSceneTracksModule.h"
 
 
 /**
diff --git a/Engine/Source/Runtime/MovieSceneTracks/Private/MovieSceneTracksPrivatePCH.h b/Engine/Source/Runtime/MovieSceneTracks/Private/MovieSceneTracksPrivatePCH.h
--- a/Engine/Source/Runtime/MovieSceneTracks/Private/MovieSceneTracksPrivatePCH.h
+++ b/Engine/Source/Runtime/MovieSceneTracks/Private/MovieSceneTracksPrivatePCH.h
@@ -3,20 +3,12 @@
 #pragma once
 
 
-/* Private dependencies
- *****************************************************************************/
-
-#include "Core.h"
-#include "Engine.h"
-#include "IMovieScenePlayer.h"
-#include "KeyParams.h"
-#include "MovieScene.h"
-#include "MovieSceneCommonHelpers.h"
-#include "MovieScenePropertyTrack.h"
-#include "MovieSceneSection.h"
-#include "MovieSceneTrack.h"
-#include "IMovieSceneTrackInstance.h"
-#include "IMovieSceneTracksModule.h"
-#include "IKeyframeSection.h"
-
-
+#include "Runtime/Core/Public/Core.h"
+#include "Runtime/Engine/Public/Engine.h"
+#include "Runtime/MovieScene/Public/IMovieScenePlayer.h"
+#include "Runtime/MovieScene/Public/IMovieSceneTrackInstance.h"
+#include "Runtime/MovieScene/Public/KeyParams.h"
+#include "Runtime/MovieScene/Public/MovieScene.h"
+#include "Runtime/MovieScene/Public/MovieSceneCommonHelpers.h"
+#include "Runtime/MovieScene/Public/MovieSceneSection.h"
+#include "Runtime/MovieScene/Public/MovieSceneTrack.h"
diff --git a/Engine/Source/Runtime/MovieSceneTracks/Private/Sections/MovieScene3DConstraintSection.cpp b/Engine/Source/Runtime/MovieSceneTracks/Private/Sections/MovieScene3DConstraintSection.cpp
--- a/Engine/Source/Runtime/MovieSceneTracks/Private/Sections/MovieScene3DConstraintSection.cpp
+++ b/Engine/Source/Runtime/MovieSceneTracks/Private/Sections/MovieScene3DConstraintSection.cpp
@@ -22,3 +22,15 @@ FGuid UMovieScene3DConstraintSection::GetConstraintId() const
 {
 	return ConstraintId;
 }
+
+
+TOptional<float> UMovieScene3DConstraintSection::GetKeyTime(FKeyHandle KeyHandle) const
+{
+	return TOptional<float>();
+}
+
+
+void UMovieScene3DConstraintSection::SetKeyTime(FKeyHandle KeyHandle, float Time)
+{
+	// do nothing
+}
diff --git a/Engine/Source/Runtime/MovieSceneTracks/Private/Sections/MovieSceneColorSection.cpp b/Engine/Source/Runtime/MovieSceneTracks/Private/Sections/MovieSceneColorSection.cpp
--- a/Engine/Source/Runtime/MovieSceneTracks/Private/Sections/MovieSceneColorSection.cpp
+++ b/Engine/Source/Runtime/MovieSceneTracks/Private/Sections/MovieSceneColorSection.cpp
@@ -250,5 +250,5 @@ bool UMovieSceneColorSection::HasKeys(const FColorKey& Key) const
 void UMovieSceneColorSection::SetDefault(const FColorKey& Key)
 {
 	FRichCurve* ChannelCurve = GetCurveForChannel(Key.Channel, &RedCurve, &GreenCurve, &BlueCurve, &AlphaCurve);
-	return SetCurveDefault(*ChannelCurve, Key.ChannelValue);
+	SetCurveDefault(*ChannelCurve, Key.ChannelValue);
 }
diff --git a/Engine/Source/Runtime/MovieSceneTracks/Private/Sections/MovieSceneLevelVisibilitySection.cpp b/Engine/Source/Runtime/MovieSceneTracks/Private/Sections/MovieSceneLevelVisibilitySection.cpp
--- a/Engine/Source/Runtime/MovieSceneTracks/Private/Sections/MovieSceneLevelVisibilitySection.cpp
+++ b/Engine/Source/Runtime/MovieSceneTracks/Private/Sections/MovieSceneLevelVisibilitySection.cpp
@@ -27,3 +27,15 @@ TArray<FName>* UMovieSceneLevelVisibilitySection::GetLevelNames()
 {
 	return &LevelNames;
 }
+
+
+TOptional<float> UMovieSceneLevelVisibilitySection::GetKeyTime(FKeyHandle KeyHandle) const
+{
+	return TOptional<float>();
+}
+
+
+void UMovieSceneLevelVisibilitySection::SetKeyTime(FKeyHandle KeyHandle, float Time)
+{
+	// do nothing
+}
diff --git a/Engine/Source/Runtime/MovieSceneTracks/Private/Sections/MovieSceneStringSection.cpp b/Engine/Source/Runtime/MovieSceneTracks/Private/Sections/MovieSceneStringSection.cpp
new file mode 100644
--- /dev/null
+++ b/Engine/Source/Runtime/MovieSceneTracks/Private/Sections/MovieSceneStringSection.cpp
@@ -0,0 +1,102 @@
+// Copyright 1998-2016 Epic Games, Inc. All Rights Reserved.
+
+#include "MovieSceneTracksPrivatePCH.h"
+#include "MovieSceneStringSection.h"
+
+UMovieSceneStringSection::UMovieSceneStringSection( const FObjectInitializer& ObjectInitializer )
+	: Super( ObjectInitializer )
+{ }
+
+/* IKeyframeSection interface
+ *****************************************************************************/
+
+FString UMovieSceneStringSection::Eval(float Time, const FString& DefaultString) const
+{
+	return StringCurve.Eval(Time, DefaultString);
+}
+
+
+/* IKeyframeSection interface
+ *****************************************************************************/
+
+bool UMovieSceneStringSection::NewKeyIsNewData(float Time, const FString& Key) const
+{
+	return (StringCurve.Eval(Time, FString()) != Key);
+}
+
+
+bool UMovieSceneStringSection::HasKeys(const FString& Key) const
+{
+	return (StringCurve.GetNumKeys() > 0);
+}
+
+void UMovieSceneStringSection::AddKey( float Time, const FString& Key, EMovieSceneKeyInterpolation KeyInterpolation )
+{
+	if (TryModify())
+	{
+		StringCurve.UpdateOrAddKey(Time, Key);
+	}
+}
+
+void UMovieSceneStringSection::SetDefault(const FString& Value)
+{
+	StringCurve.SetDefaultValue(Value);
+}
+
+
+/* UMovieSceneSection interface
+ *****************************************************************************/
+
+void UMovieSceneStringSection::DilateSection(float DilationFactor, float Origin, TSet<FKeyHandle>& KeyHandles)
+{
+	Super::DilateSection(DilationFactor, Origin, KeyHandles);
+
+	StringCurve.ScaleCurve(Origin, DilationFactor, KeyHandles);
+}
+
+
+void UMovieSceneStringSection::GetKeyHandles(TSet<FKeyHandle>& OutKeyHandles, TRange<float> TimeRange) const
+{
+	if (!TimeRange.Overlaps(GetRange()))
+	{
+		return;
+	}
+
+	for (auto It(StringCurve.GetKeyHandleIterator()); It; ++It)
+	{
+		float Time = StringCurve.GetKeyTime(It.Key());
+
+		if (TimeRange.Contains(Time))
+		{
+			OutKeyHandles.Add(It.Key());
+		}
+	}
+}
+
+
+TOptional<float> UMovieSceneStringSection::GetKeyTime(FKeyHandle KeyHandle) const
+{
+	if (StringCurve.IsKeyHandleValid(KeyHandle))
+	{
+		return TOptional<float>(StringCurve.GetKeyTime(KeyHandle));
+	}
+
+	return TOptional<float>();
+}
+
+
+void UMovieSceneStringSection::MoveSection(float DeltaPosition, TSet<FKeyHandle>& KeyHandles)
+{
+	Super::MoveSection(DeltaPosition, KeyHandles);
+
+	StringCurve.ShiftCurve(DeltaPosition, KeyHandles);
+}
+
+
+void UMovieSceneStringSection::SetKeyTime(FKeyHandle KeyHandle, float Time)
+{
+	if (StringCurve.IsKeyHandleValid(KeyHandle))
+	{
+		StringCurve.SetKeyTime(KeyHandle, Time);
+	}
+}
diff --git a/Engine/Source/Runtime/MovieSceneTracks/Private/TrackInstances/MovieScene3DTransformTrackInstance.cpp b/Engine/Source/Runtime/MovieSceneTracks/Private/TrackInstances/MovieScene3DTransformTrackInstance.cpp
--- a/Engine/Source/Runtime/MovieSceneTracks/Private/TrackInstances/MovieScene3DTransformTrackInstance.cpp
+++ b/Engine/Source/Runtime/MovieSceneTracks/Private/TrackInstances/MovieScene3DTransformTrackInstance.cpp
@@ -106,7 +106,9 @@ void FMovieScene3DTransformTrackInstance::Update(EMovieSceneUpdateData& UpdateDa
 			{
 				if (UpdateData.UpdatePass == MSUP_PreUpdate)
 				{
-					SceneComponent->ResetRelativeTransform();
+					SceneComponent->RelativeLocation = FVector(0, 0, 0);
+					SceneComponent->RelativeRotation = FRotator(0, 0, 0);
+					SceneComponent->RelativeScale3D = FVector(1.f, 1.f, 1.f);
 				}
 				else if (UpdateData.UpdatePass == MSUP_Update)
 				{
diff --git a/Engine/Source/Runtime/MovieSceneTracks/Private/TrackInstances/MovieSceneStringTrackInstance.cpp b/Engine/Source/Runtime/MovieSceneTracks/Private/TrackInstances/MovieSceneStringTrackInstance.cpp
new file mode 100644
--- /dev/null
+++ b/Engine/Source/Runtime/MovieSceneTracks/Private/TrackInstances/MovieSceneStringTrackInstance.cpp
@@ -0,0 +1,78 @@
+// Copyright 1998-2016 Epic Games, Inc. All Rights Reserved.
+
+#include "MovieSceneTracksPrivatePCH.h"
+#include "MovieSceneStringTrack.h"
+#include "MovieSceneStringTrackInstance.h"
+
+
+/* FMovieSceneStringTrackInstance structors
+ *****************************************************************************/
+
+FMovieSceneStringTrackInstance::FMovieSceneStringTrackInstance(UMovieSceneStringTrack& InStringTrack)
+{
+	StringTrack = &InStringTrack;
+
+	PropertyBindings = MakeShareable(new FTrackInstancePropertyBindings(StringTrack->GetPropertyName(), StringTrack->GetPropertyPath()));
+}
+
+
+/* IMovieSceneTrackInstance interface
+ *****************************************************************************/
+
+void FMovieSceneStringTrackInstance::RefreshInstance(const TArray<TWeakObjectPtr<UObject>>& RuntimeObjects, class IMovieScenePlayer& Player, FMovieSceneSequenceInstance& SequenceInstance)
+{
+	PropertyBindings->UpdateBindings(RuntimeObjects);
+}
+
+
+void FMovieSceneStringTrackInstance::RestoreState(const TArray<TWeakObjectPtr<UObject>>& RuntimeObjects, IMovieScenePlayer& Player, FMovieSceneSequenceInstance& SequenceInstance)
+{
+	for (auto ObjectPtr : RuntimeObjects)
+	{
+		UObject* Object = ObjectPtr.Get();
+
+		if (!IsValid(Object))
+		{
+			continue;
+		}
+
+		FString* StringValue = InitStringMap.Find(Object);
+
+		if (StringValue != nullptr)
+		{
+			PropertyBindings->CallFunction<FString>(Object, StringValue);
+		}
+	}
+
+	PropertyBindings->UpdateBindings(RuntimeObjects);
+}
+
+
+void FMovieSceneStringTrackInstance::SaveState(const TArray<TWeakObjectPtr<UObject>>& RuntimeObjects, IMovieScenePlayer& Player, FMovieSceneSequenceInstance& SequenceInstance)
+{
+	for (auto ObjectPtr : RuntimeObjects)
+	{
+		UObject* Object = ObjectPtr.Get();
+
+		if (InitStringMap.Find(Object) == nullptr)
+		{
+			const FString StringValue = PropertyBindings->GetCurrentValue<FString>(Object);
+			InitStringMap.Add(Object, StringValue);
+		}
+	}
+}
+
+
+void FMovieSceneStringTrackInstance::Update(EMovieSceneUpdateData& UpdateData, const TArray<TWeakObjectPtr<UObject>>& RuntimeObjects, IMovieScenePlayer& Player, FMovieSceneSequenceInstance& SequenceInstance)
+{
+	FString StringValue;
+
+	if (StringTrack->Eval(UpdateData.Position, UpdateData.LastPosition, StringValue))
+	{
+		for (auto ObjectPtr : RuntimeObjects)
+		{
+			UObject* Object = ObjectPtr.Get();
+			PropertyBindings->CallFunction<FString>(Object, &StringValue);
+		}
+	}
+}
diff --git a/Engine/Source/Runtime/MovieSceneTracks/Private/Tracks/MovieSceneStringTrack.cpp b/Engine/Source/Runtime/MovieSceneTracks/Private/Tracks/MovieSceneStringTrack.cpp
new file mode 100644
--- /dev/null
+++ b/Engine/Source/Runtime/MovieSceneTracks/Private/Tracks/MovieSceneStringTrack.cpp
@@ -0,0 +1,133 @@
+// Copyright 1998-2016 Epic Games, Inc. All Rights Reserved.
+
+#include "MovieSceneTracksPrivatePCH.h"
+#include "MovieSceneStringSection.h"
+#include "MovieSceneStringTrack.h"
+#include "MovieSceneStringTrackInstance.h"
+
+
+#define LOCTEXT_NAMESPACE "MovieSceneStringTrack"
+
+
+/* UMovieSceneStringTrack interface
+ *****************************************************************************/
+
+bool UMovieSceneStringTrack::AddKeyToSection(float Time, const FString& String, FKeyParams KeyParams)
+{
+	UMovieSceneSection* TargetSection = MovieSceneHelpers::FindNearestSectionAtTime(Sections, Time);
+
+	if (TargetSection == nullptr)
+	{
+		TargetSection = CreateNewSection();
+		TargetSection->SetStartTime(Time);
+		TargetSection->SetEndTime(Time);
+
+		Sections.Add(TargetSection);
+	}
+
+	UMovieSceneStringSection* StringSection = Cast<UMovieSceneStringSection>(TargetSection);
+
+	if (StringSection == nullptr)
+	{
+		return false;
+	}
+
+	StringSection->AddKey(Time, String, EMovieSceneKeyInterpolation::Auto);
+
+	return true;
+}
+
+
+bool UMovieSceneStringTrack::Eval(float Position, float LastPostion, FString& OutString) const
+{
+	const UMovieSceneSection* Section = MovieSceneHelpers::FindNearestSectionAtTime(Sections, Position);
+
+	if (Section)
+	{
+		if (!Section->IsInfinite())
+		{
+			Position = FMath::Clamp(Position, Section->GetStartTime(), Section->GetEndTime());
+		}
+
+		OutString = CastChecked<UMovieSceneStringSection>(Section)->Eval(Position, OutString);
+	}
+
+	return Section != nullptr;
+}
+
+
+/* UMovieSceneTrack interface
+ *****************************************************************************/
+
+void UMovieSceneStringTrack::AddSection(UMovieSceneSection& Section)
+{
+	Sections.Add(&Section);
+}
+
+
+TSharedPtr<IMovieSceneTrackInstance> UMovieSceneStringTrack::CreateInstance()
+{
+	return MakeShareable(new FMovieSceneStringTrackInstance(*this));
+}
+
+
+UMovieSceneSection* UMovieSceneStringTrack::CreateNewSection()
+{
+	return NewObject<UMovieSceneSection>(this, UMovieSceneStringSection::StaticClass(), NAME_None, RF_Transactional);
+}
+
+
+const TArray<UMovieSceneSection*>& UMovieSceneStringTrack::GetAllSections() const
+{
+	return Sections;
+}
+
+
+TRange<float> UMovieSceneStringTrack::GetSectionBoundaries() const
+{
+	TRange<float> SectionBoundaries = TRange<float>::Empty();
+
+	for (auto& Section : Sections)
+	{
+		SectionBoundaries = TRange<float>::Hull(SectionBoundaries, Section->GetRange());
+	}
+
+	return SectionBoundaries;
+}
+
+
+bool UMovieSceneStringTrack::HasSection(const UMovieSceneSection& Section) const
+{
+	return Sections.Contains(&Section);
+}
+
+
+bool UMovieSceneStringTrack::IsEmpty() const
+{
+	return (Sections.Num() == 0);
+}
+
+
+void UMovieSceneStringTrack::RemoveAllAnimationData()
+{
+	Sections.Empty();
+}
+
+
+void UMovieSceneStringTrack::RemoveSection(UMovieSceneSection& Section)
+{
+	Sections.Remove(&Section);
+}
+
+
+#if WITH_EDITORONLY_DATA
+
+FText UMovieSceneStringTrack::GetDefaultDisplayName() const
+{ 
+	return LOCTEXT("TrackName", "Strings"); 
+}
+
+#endif
+
+
+#undef LOCTEXT_NAMESPACE
diff --git a/Engine/Source/Runtime/MovieSceneTracks/Public/IMovieSceneTracksModule.h b/Engine/Source/Runtime/MovieSceneTracks/Public/IMovieSceneTracksModule.h
--- a/Engine/Source/Runtime/MovieSceneTracks/Public/IMovieSceneTracksModule.h
+++ b/Engine/Source/Runtime/MovieSceneTracks/Public/IMovieSceneTracksModule.h
@@ -3,7 +3,7 @@
 #pragma once
 
 #include "ModuleInterface.h"
-#include "ModuleManager.h"		// For inline LoadModuleChecked()
+#include "ModuleManager.h"
 
 
 /**
diff --git a/Engine/Source/Runtime/MovieSceneTracks/Public/Sections/MovieScene3DAttachSection.h b/Engine/Source/Runtime/MovieSceneTracks/Public/Sections/MovieScene3DAttachSection.h
--- a/Engine/Source/Runtime/MovieSceneTracks/Public/Sections/MovieScene3DAttachSection.h
+++ b/Engine/Source/Runtime/MovieSceneTracks/Public/Sections/MovieScene3DAttachSection.h
@@ -5,6 +5,11 @@
 #include "MovieScene3DConstraintSection.h"
 #include "MovieScene3DAttachSection.generated.h"
 
+
+class AActor;
+class USceneComponent;
+
+
 /**
  * A 3D Attach section
  */
@@ -19,9 +24,9 @@ public:
 	/**
 	 * Evaluates the attach track
 	 *
-	 * @param Time				The position in time within the movie scene
+	 * @param Time The position in time within the movie scene
 	 */
-	void Eval( USceneComponent* SceneComponent, float Time, AActor* Actor, FVector& OutTranslation, FRotator& OutRotation ) const;
+	void Eval(USceneComponent* SceneComponent, float Time, AActor* Actor, FVector& OutTranslation, FRotator& OutRotation) const;
 
 	/** 
 	 * Adds an attach to the section
@@ -30,8 +35,10 @@ public:
 	 * @param SequenceEndTime   The time at the end of the sequence, by default the attach is set to end at this time
 	 * @param InAttachId The id to the path
 	 */
-	void AddAttach( float Time, float SequenceEndTime, const FGuid& InAttachId);
+	void AddAttach(float Time, float SequenceEndTime, const FGuid& InAttachId);
 	
+public:
+
 	UPROPERTY(EditAnywhere, Category="Attach")
 	FName AttachSocketName;
 
@@ -55,4 +62,4 @@ public:
 
 	UPROPERTY(EditAnywhere, Category="Attach")
 	uint32 bConstrainRz:1;
-};
\ No newline at end of file
+};
diff --git a/Engine/Source/Runtime/MovieSceneTracks/Public/Sections/MovieScene3DConstraintSection.h b/Engine/Source/Runtime/MovieSceneTracks/Public/Sections/MovieScene3DConstraintSection.h
--- a/Engine/Source/Runtime/MovieSceneTracks/Public/Sections/MovieScene3DConstraintSection.h
+++ b/Engine/Source/Runtime/MovieSceneTracks/Public/Sections/MovieScene3DConstraintSection.h
@@ -5,6 +5,7 @@
 #include "MovieSceneSection.h"
 #include "MovieScene3DConstraintSection.generated.h"
 
+
 /**
  * Base class for 3D constraint section
  */
@@ -22,12 +23,16 @@ public:
 	/** Gets the constraint id for this section */
 	virtual FGuid GetConstraintId() const;
 
-	// UMovieSceneSection interface
-	virtual TOptional<float> GetKeyTime( FKeyHandle KeyHandle ) const override { return TOptional<float>(); }
-	virtual void SetKeyTime( FKeyHandle KeyHandle, float Time ) override { }
+public:
+
+	//~ UMovieSceneSection interface
+
+	virtual TOptional<float> GetKeyTime(FKeyHandle KeyHandle) const override;
+	virtual void SetKeyTime(FKeyHandle KeyHandle, float Time) override;
 	
 protected:
+
 	/** The possessable guid that this constraint uses */
 	UPROPERTY()
 	FGuid ConstraintId;
-};
\ No newline at end of file
+};
diff --git a/Engine/Source/Runtime/MovieSceneTracks/Public/Sections/MovieScene3DPathSection.h b/Engine/Source/Runtime/MovieSceneTracks/Public/Sections/MovieScene3DPathSection.h
--- a/Engine/Source/Runtime/MovieSceneTracks/Public/Sections/MovieScene3DPathSection.h
+++ b/Engine/Source/Runtime/MovieSceneTracks/Public/Sections/MovieScene3DPathSection.h
@@ -5,8 +5,11 @@
 #include "MovieScene3DConstraintSection.h"
 #include "MovieScene3DPathSection.generated.h"
 
+
+class USceneComponent;
 class USplineComponent;
 
+
 UENUM(BlueprintType)
 enum class MovieScene3DPathSection_Axis : uint8
 {
@@ -18,6 +21,7 @@ enum class MovieScene3DPathSection_Axis : uint8
 	NEG_Z UMETA(DisplayName = "-Z")
 };
 
+
 /**
  * A 3D Path section
  */
@@ -29,35 +33,42 @@ class UMovieScene3DPathSection
 
 public:
 
-	/** MovieSceneSection interface */
-	virtual void MoveSection( float DeltaPosition, TSet<FKeyHandle>& KeyHandles ) override;
-	virtual void DilateSection( float DilationFactor, float Origin, TSet<FKeyHandle>& KeyHandles ) override;
-	virtual void GetKeyHandles(TSet<FKeyHandle>& OutKeyHandles, TRange<float> TimeRange) const override;
-
 	/**
-	 * Evaluates the path track
+	 * Evaluates the path track.
 	 *
-	 * @param Time				The position in time within the movie scene
+	 * @param Time The position in time within the movie scene.
 	 */
-	void Eval( USceneComponent* SceneComponent, float Time, USplineComponent* SplineComponent, FVector& OutTranslation, FRotator& OutRotation ) const;
+	void Eval(USceneComponent* SceneComponent, float Time, USplineComponent* SplineComponent, FVector& OutTranslation, FRotator& OutRotation) const;
 
 	/** 
-	 * Adds a path to the section
+	 * Adds a path to the section.
 	 *
-	 * @param Time	The location in time where the path should be added
-	 * @param SequenceEndTime   The time at the end of the sequence, by default the path is set to end at this time
-	 * @param InPathId The id to the path
+	 * @param Time The location in time where the path should be added.
+	 * @param SequenceEndTime The time at the end of the sequence, by default the path is set to end at this time.
+	 * @param InPathId The id to the path.
 	 */
-	void AddPath( float Time, float SequenceEndTime, const FGuid& InPathId);
+	void AddPath(float Time, float SequenceEndTime, const FGuid& InPathId);
 
 	/** 
-	 * Returns the timing curve
+	 * Returns the timing curve.
 	 *
-	 * @return The timing curve
+	 * @return The timing curve.
 	 */
-	MOVIESCENETRACKS_API FRichCurve& GetTimingCurve() { return TimingCurve; }
+	MOVIESCENETRACKS_API FRichCurve& GetTimingCurve()
+	{
+		return TimingCurve;
+	}
 	
+public:
+
+	//~ MovieSceneSection interface
+
+	virtual void MoveSection(float DeltaPosition, TSet<FKeyHandle>& KeyHandles) override;
+	virtual void DilateSection(float DilationFactor, float Origin, TSet<FKeyHandle>& KeyHandles) override;
+	virtual void GetKeyHandles(TSet<FKeyHandle>& OutKeyHandles, TRange<float> TimeRange) const override;
+
 private:
+
 	/** Timing Curve */
 	UPROPERTY(EditAnywhere, Category="Path")
 	FRichCurve TimingCurve;
@@ -81,4 +92,4 @@ private:
 	/** Force Upright */
 	UPROPERTY(EditAnywhere, Category="Path")
 	uint32 bForceUpright:1;
-};
\ No newline at end of file
+};
diff --git a/Engine/Source/Runtime/MovieSceneTracks/Public/Sections/MovieSceneActorReferenceSection.h b/Engine/Source/Runtime/MovieSceneTracks/Public/Sections/MovieSceneActorReferenceSection.h
--- a/Engine/Source/Runtime/MovieSceneTracks/Public/Sections/MovieSceneActorReferenceSection.h
+++ b/Engine/Source/Runtime/MovieSceneTracks/Public/Sections/MovieSceneActorReferenceSection.h
@@ -2,52 +2,65 @@
 
 #pragma once
 
+#include "Curves/IntegralCurve.h"
 #include "MovieSceneSection.h"
 #include "IKeyframeSection.h"
 #include "MovieSceneActorReferenceSection.generated.h"
 
+
 /**
  * A single actor reference point section
  */
-UCLASS( MinimalAPI )
+UCLASS(MinimalAPI)
 class UMovieSceneActorReferenceSection
 	: public UMovieSceneSection
 	, public IKeyframeSection<FGuid>
 {
 	GENERATED_UCLASS_BODY()
+
 public:
+
 	/**
 	 * Updates this section
 	 *
 	 * @param Position	The position in time within the movie scene
 	 */
-	FGuid Eval( float Position ) const;
+	FGuid Eval(float Position) const;
 
 	// IKeyframeSection interface.
 
-	void AddKey( float Time, const FGuid& Value, EMovieSceneKeyInterpolation KeyInterpolation );
+	void AddKey(float Time, const FGuid& Value, EMovieSceneKeyInterpolation KeyInterpolation);
 	bool NewKeyIsNewData(float Time, const FGuid& Value) const;
-	bool HasKeys( const FGuid& Value ) const;
-	void SetDefault( const FGuid& Value );
+	bool HasKeys(const FGuid& Value) const;
+	void SetDefault(const FGuid& Value);
 
 	/**
-	 * UMovieSceneSection interface 
+	 * @return The integral curve on this section
 	 */
+	FIntegralCurve& GetActorReferenceCurve()
+	{
+		return ActorGuidIndexCurve;
+	}
+
+public:
+
+	//~ UMovieSceneSection interface
+
 	virtual void MoveSection(float DeltaPosition, TSet<FKeyHandle>& KeyHandles) override;
 	virtual void DilateSection(float DilationFactor, float Origin, TSet<FKeyHandle>& KeyHandles) override;
 	virtual void GetKeyHandles(TSet<FKeyHandle>& OutKeyHandles, TRange<float> TimeRange) const override;
-	virtual TOptional<float> GetKeyTime( FKeyHandle KeyHandle ) const override;
-	virtual void SetKeyTime( FKeyHandle KeyHandle, float Time ) override;
+	virtual TOptional<float> GetKeyTime(FKeyHandle KeyHandle) const override;
+	virtual void SetKeyTime(FKeyHandle KeyHandle, float Time) override;
 
-	/**
-	 * @return The integral curve on this section
-	 */
-	FIntegralCurve& GetActorReferenceCurve() { return ActorGuidIndexCurve; }
+public:
+
+	//~ UObject interface
 
 	virtual void PreSave(const class ITargetPlatform* TargetPlatform) override;
 	virtual void PostLoad() override;
 
 private:
+
 	/** Curve data */
 	UPROPERTY()
 	FIntegralCurve ActorGuidIndexCurve;
@@ -56,4 +69,4 @@ private:
 
 	UPROPERTY()
 	TArray<FString> ActorGuidStrings;
-};
\ No newline at end of file
+};
diff --git a/Engine/Source/Runtime/MovieSceneTracks/Public/Sections/MovieSceneBoolSection.h b/Engine/Source/Runtime/MovieSceneTracks/Public/Sections/MovieSceneBoolSection.h
--- a/Engine/Source/Runtime/MovieSceneTracks/Public/Sections/MovieSceneBoolSection.h
+++ b/Engine/Source/Runtime/MovieSceneTracks/Public/Sections/MovieSceneBoolSection.h
@@ -2,14 +2,16 @@
 
 #pragma once
 
+#include "Curves/IntegralCurve.h"
 #include "MovieSceneSection.h"
 #include "IKeyframeSection.h"
 #include "MovieSceneBoolSection.generated.h"
 
+
 /**
- * A single bool section
+ * A single bool section.
  */
-UCLASS( MinimalAPI )
+UCLASS(MinimalAPI)
 class UMovieSceneBoolSection 
 	: public UMovieSceneSection
 	, public IKeyframeSection<bool>
@@ -25,33 +27,45 @@ public:
 public:
 
 	/**
-	 * Updates this section
+	 * Update this section.
 	 *
-	 * @param Position	The position in time within the movie scene
+	 * @param Position The position in time within the movie scene.
 	 */
-	virtual bool Eval( float Position ) const;
+	virtual bool Eval(float Position) const;
 
-	// IKeyframeSection interface.
-	virtual void AddKey( float Time, const bool& Value, EMovieSceneKeyInterpolation KeyInterpolation ) override;
-	virtual void SetDefault( const bool& Value ) override;
-	virtual bool NewKeyIsNewData(float Time, const bool& Value ) const override;
-	virtual bool HasKeys( const bool& Value ) const;
+	/** Gets all the keys of this boolean section. */
+	FIntegralCurve& GetCurve()
+	{
+		return BoolCurve;
+	}
+
+public:
+
+	//~ IKeyframeSection interface
+
+	virtual void AddKey(float Time, const bool& Value, EMovieSceneKeyInterpolation KeyInterpolation) override;
+	virtual void SetDefault(const bool& Value) override;
+	virtual bool NewKeyIsNewData(float Time, const bool& Value) const override;
+	virtual bool HasKeys(const bool& Value) const;
+
+public:
+
+	//~ UMovieSceneSection interface 
 
-	/**
-	 * UMovieSceneSection interface 
-	 */
 	virtual void MoveSection(float DeltaPosition, TSet<FKeyHandle>& KeyHandles) override;
 	virtual void DilateSection(float DilationFactor, float Origin, TSet<FKeyHandle>& KeyHandles) override;
 	virtual void GetKeyHandles(TSet<FKeyHandle>& OutKeyHandles, TRange<float> TimeRange) const override;
-	virtual TOptional<float> GetKeyTime( FKeyHandle KeyHandle ) const override;
-	virtual void SetKeyTime( FKeyHandle KeyHandle, float Time ) override;
+	virtual TOptional<float> GetKeyTime(FKeyHandle KeyHandle) const override;
+	virtual void SetKeyTime(FKeyHandle KeyHandle, float Time) override;
 
-	/** Gets all the keys of this boolean section */
-	FIntegralCurve& GetCurve() { return BoolCurve; }
+public:
+
+	//~ UObject interface
 
 	virtual void PostLoad() override;
 
 private:
+
 	/** Ordered curve data */
 	// @todo Sequencer This could be optimized by packing the bools separately
 	// but that may not be worth the effort
diff --git a/Engine/Source/Runtime/MovieSceneTracks/Public/Sections/MovieSceneByteSection.h b/Engine/Source/Runtime/MovieSceneTracks/Public/Sections/MovieSceneByteSection.h
--- a/Engine/Source/Runtime/MovieSceneTracks/Public/Sections/MovieSceneByteSection.h
+++ b/Engine/Source/Runtime/MovieSceneTracks/Public/Sections/MovieSceneByteSection.h
@@ -2,14 +2,16 @@
 
 #pragma once
 
+#include "Curves/IntegralCurve.h"
 #include "MovieSceneSection.h"
 #include "IKeyframeSection.h"
 #include "MovieSceneByteSection.generated.h"
 
+
 /**
- * A single byte section
+ * A single byte section.
  */
-UCLASS( MinimalAPI )
+UCLASS(MinimalAPI)
 class UMovieSceneByteSection 
 	: public UMovieSceneSection
 	, public IKeyframeSection<uint8>
@@ -17,31 +19,39 @@ class UMovieSceneByteSection
 	GENERATED_UCLASS_BODY()
 public:
 	/**
-	 * Updates this section
+	 * Update this section.
 	 *
-	 * @param Position	The position in time within the movie scene
+	 * @param Position The position in time within the movie scene
 	 */
-	virtual uint8 Eval( float Position ) const;
+	virtual uint8 Eval(float Position) const;
 
-	// IKeyframeSection interface.
-	virtual void AddKey( float Time, const uint8& Value, EMovieSceneKeyInterpolation KeyInterpolation ) override;
-	virtual bool NewKeyIsNewData( float Time, const uint8& Value ) const override;
-	virtual bool HasKeys( const uint8& Value ) const override;
-	virtual void SetDefault( const uint8& Value ) override;
+	/** Gets all the keys of this byte section */
+	FIntegralCurve& GetCurve()
+	{
+		return ByteCurve;
+	}
+
+public:
+
+	//~ IKeyframeSection interface
+
+	virtual void AddKey(float Time, const uint8& Value, EMovieSceneKeyInterpolation KeyInterpolation) override;
+	virtual bool NewKeyIsNewData(float Time, const uint8& Value) const override;
+	virtual bool HasKeys(const uint8& Value) const override;
+	virtual void SetDefault(const uint8& Value) override;
+
+public:
+
+	//~ UMovieSceneSection interface 
 
-	/**
-	 * UMovieSceneSection interface 
-	 */
 	virtual void MoveSection(float DeltaPosition, TSet<FKeyHandle>& KeyHandles) override;
 	virtual void DilateSection(float DilationFactor, float Origin, TSet<FKeyHandle>& KeyHandles) override;
 	virtual void GetKeyHandles(TSet<FKeyHandle>& OutKeyHandles, TRange<float> TimeRange) const override;
-	virtual TOptional<float> GetKeyTime( FKeyHandle KeyHandle ) const override;
-	virtual void SetKeyTime( FKeyHandle KeyHandle, float Time ) override;
-
-	/** Gets all the keys of this byte section */
-	FIntegralCurve& GetCurve() { return ByteCurve; }
+	virtual TOptional<float> GetKeyTime(FKeyHandle KeyHandle) const override;
+	virtual void SetKeyTime(FKeyHandle KeyHandle, float Time) override;
 
 private:
+
 	/** Ordered curve data */
 	// @todo Sequencer This could be optimized by packing the bytes separately
 	// but that may not be worth the effort
diff --git a/Engine/Source/Runtime/MovieSceneTracks/Public/Sections/MovieSceneColorSection.h b/Engine/Source/Runtime/MovieSceneTracks/Public/Sections/MovieSceneColorSection.h
--- a/Engine/Source/Runtime/MovieSceneTracks/Public/Sections/MovieSceneColorSection.h
+++ b/Engine/Source/Runtime/MovieSceneTracks/Public/Sections/MovieSceneColorSection.h
@@ -67,7 +67,8 @@ public:
 	/**
 	 * Updates this section
 	 *
-	 * @param Position	The position in time within the movie scene
+	 * @param Position The position in time within the movie scene.
+	 * @param DefaultColor The default value to return.
 	 */
 	virtual FLinearColor Eval(float Position, const FLinearColor& DefaultColor) const;
 
@@ -75,6 +76,7 @@ public:
 	 * Gets the red color curve
 	 *
 	 * @return The rich curve for this color channel
+	 * @see GetAlphaCurve, GetBlueCurve, GetGreenCurve
 	 */
 	FRichCurve& GetRedCurve() { return RedCurve; }
 	const FRichCurve& GetRedCurve() const { return RedCurve; }
@@ -83,6 +85,7 @@ public:
 	 * Gets the green color curve
 	 *
 	 * @return The rich curve for this color channel
+	 * @see GetAlphaCurve, GetBlueCurve, GetRedCurve
 	 */
 	FRichCurve& GetGreenCurve() { return GreenCurve; }
 	const FRichCurve& GetGreenCurve() const { return GreenCurve; }
@@ -91,6 +94,7 @@ public:
 	 * Gets the blue color curve
 	 *
 	 * @return The rich curve for this color channel
+	 * @see GetAlphaCurve, GetGreenCurve, GetRedCurve
 	 */
 	FRichCurve& GetBlueCurve() { return BlueCurve; }
 	const FRichCurve& GetBlueCurve() const { return BlueCurve; }
@@ -99,24 +103,25 @@ public:
 	 * Gets the alpha color curve
 	 *
 	 * @return The rich curve for this color channel
+	 * @see GetBlueCurve, GetGreenCurve, GetRedCurve
 	 */
 	FRichCurve& GetAlphaCurve() { return AlphaCurve; }
 	const FRichCurve& GetAlphaCurve() const { return AlphaCurve; }
 
 public:
 
-	// UMovieSceneSection interface
+	//~ UMovieSceneSection interface
 
 	virtual void MoveSection(float DeltaPosition, TSet<FKeyHandle>& KeyHandles) override;
 	virtual void DilateSection(float DilationFactor, float Origin, TSet<FKeyHandle>& KeyHandles) override;
 	virtual void GetKeyHandles(TSet<FKeyHandle>& OutKeyHandles, TRange<float> TimeRange) const override;
 	virtual TSharedPtr<FStructOnScope> GetKeyStruct(const TArray<FKeyHandle>& KeyHandles) override;
-	virtual TOptional<float> GetKeyTime( FKeyHandle KeyHandle ) const override;
-	virtual void SetKeyTime( FKeyHandle KeyHandle, float Time ) override;
+	virtual TOptional<float> GetKeyTime(FKeyHandle KeyHandle) const override;
+	virtual void SetKeyTime(FKeyHandle KeyHandle, float Time) override;
 
 public:
 
-	// IKeyframeSection interface
+	//~ IKeyframeSection interface
 
 	virtual void AddKey(float Time, const FColorKey& Key, EMovieSceneKeyInterpolation KeyInterpolation) override;
 	virtual bool NewKeyIsNewData(float Time, const FColorKey& Key) const override;
diff --git a/Engine/Source/Runtime/MovieSceneTracks/Public/Sections/MovieSceneEventSection.h b/Engine/Source/Runtime/MovieSceneTracks/Public/Sections/MovieSceneEventSection.h
--- a/Engine/Source/Runtime/MovieSceneTracks/Public/Sections/MovieSceneEventSection.h
+++ b/Engine/Source/Runtime/MovieSceneTracks/Public/Sections/MovieSceneEventSection.h
@@ -2,6 +2,7 @@
 
 #pragma once
 
+#include "Curves/NameCurve.h"
 #include "MovieSceneEventSection.generated.h"
 
 
@@ -11,7 +12,7 @@ class ALevelScriptActor;
 /**
  * Implements a section in movie scene event tracks.
  */
-UCLASS( MinimalAPI )
+UCLASS(MinimalAPI)
 class UMovieSceneEventSection
 	: public UMovieSceneSection
 {
@@ -32,7 +33,9 @@ public:
 	void AddKey(float Time, const FName& EventName, FKeyParams KeyParams);
 
 	/**
-	 * @return The float curve on this section
+	 * Get the section's event curve.
+	 *
+	 * @return Event curve.
 	 */
 	FNameCurve& GetEventCurve()
 	{
@@ -50,13 +53,13 @@ public:
 
 public:
 
-	// UMovieSceneSection interface
+	//~ UMovieSceneSection interface
 
 	virtual void DilateSection(float DilationFactor, float Origin, TSet<FKeyHandle>& KeyHandles) override;
 	virtual void GetKeyHandles(TSet<FKeyHandle>& KeyHandles, TRange<float> TimeRange) const override;
 	virtual void MoveSection(float DeltaPosition, TSet<FKeyHandle>& KeyHandles) override;
-	virtual TOptional<float> GetKeyTime( FKeyHandle KeyHandle ) const override;
-	virtual void SetKeyTime( FKeyHandle KeyHandle, float Time ) override;
+	virtual TOptional<float> GetKeyTime(FKeyHandle KeyHandle) const override;
+	virtual void SetKeyTime(FKeyHandle KeyHandle, float Time) override;
 
 protected:
 
diff --git a/Engine/Source/Runtime/MovieSceneTracks/Public/Sections/MovieSceneFadeSection.h b/Engine/Source/Runtime/MovieSceneTracks/Public/Sections/MovieSceneFadeSection.h
--- a/Engine/Source/Runtime/MovieSceneTracks/Public/Sections/MovieSceneFadeSection.h
+++ b/Engine/Source/Runtime/MovieSceneTracks/Public/Sections/MovieSceneFadeSection.h
@@ -7,7 +7,7 @@
 
 
 /**
- * A single floating point section
+ * A single floating point section.
  */
 UCLASS(MinimalAPI)
 class UMovieSceneFadeSection
@@ -20,11 +20,11 @@ class UMovieSceneFadeSection
 
 public:
 
-	/** Fade color */
+	/** Fade color. */
 	UPROPERTY(EditAnywhere, Category="Fade", meta=(InlineColorPicker))
 	FLinearColor FadeColor;
 
-	/** Fade audio */
+	/** Fade audio. */
 	UPROPERTY(EditAnywhere, Category="Fade")
 	uint32 bFadeAudio:1;
 };
diff --git a/Engine/Source/Runtime/MovieSceneTracks/Public/Sections/MovieSceneFloatSection.h b/Engine/Source/Runtime/MovieSceneTracks/Public/Sections/MovieSceneFloatSection.h
--- a/Engine/Source/Runtime/MovieSceneTracks/Public/Sections/MovieSceneFloatSection.h
+++ b/Engine/Source/Runtime/MovieSceneTracks/Public/Sections/MovieSceneFloatSection.h
@@ -2,10 +2,12 @@
 
 #pragma once
 
+#include "Curves/RichCurve.h"
 #include "MovieSceneSection.h"
 #include "IKeyframeSection.h"
 #include "MovieSceneFloatSection.generated.h"
 
+
 /**
  * A single floating point section
  */
@@ -15,6 +17,7 @@ class UMovieSceneFloatSection
 	, public IKeyframeSection<float>
 {
 	GENERATED_UCLASS_BODY()
+
 public:
 	/**
 	 * Updates this section
@@ -23,28 +26,36 @@ public:
 	 */
 	float Eval( float Position ) const;
 
-	// IKeyframeSection interface.
+	/**
+	 * @return The float curve on this section
+	 */
+	FRichCurve& GetFloatCurve()
+	{
+		return FloatCurve;
+	}
+
+public:
+
+	//~ IKeyframeSection interface
 
 	void AddKey( float Time, const float& Value, EMovieSceneKeyInterpolation KeyInterpolation );
 	bool NewKeyIsNewData(float Time, const float& Value) const;
 	bool HasKeys( const float& Value ) const;
 	void SetDefault( const float& Value );
 
-	/**
-	 * UMovieSceneSection interface 
-	 */
+public:
+
+	//~ UMovieSceneSection interface
+
 	virtual void MoveSection(float DeltaPosition, TSet<FKeyHandle>& KeyHandles) override;
 	virtual void DilateSection(float DilationFactor, float Origin, TSet<FKeyHandle>& KeyHandles) override;
 	virtual void GetKeyHandles(TSet<FKeyHandle>& OutKeyHandles, TRange<float> TimeRange) const override;
 	virtual TOptional<float> GetKeyTime( FKeyHandle KeyHandle ) const override;
 	virtual void SetKeyTime( FKeyHandle KeyHandle, float Time ) override;
 
-	/**
-	 * @return The float curve on this section
-	 */
-	FRichCurve& GetFloatCurve() { return FloatCurve; }
 private:
+	
 	/** Curve data */
 	UPROPERTY()
 	FRichCurve FloatCurve;
-};
\ No newline at end of file
+};
diff --git a/Engine/Source/Runtime/MovieSceneTracks/Public/Sections/MovieSceneLevelVisibilitySection.h b/Engine/Source/Runtime/MovieSceneTracks/Public/Sections/MovieSceneLevelVisibilitySection.h
--- a/Engine/Source/Runtime/MovieSceneTracks/Public/Sections/MovieSceneLevelVisibilitySection.h
+++ b/Engine/Source/Runtime/MovieSceneTracks/Public/Sections/MovieSceneLevelVisibilitySection.h
@@ -4,7 +4,10 @@
 
 #include "MovieSceneLevelVisibilitySection.generated.h"
 
-/** Visibility options for the level visibility section. */
+
+/**
+ * Visibility options for the level visibility section.
+ */
 UENUM()
 enum class ELevelVisibility
 {
@@ -15,6 +18,7 @@ enum class ELevelVisibility
 	Hidden
 };
 
+
 /**
  * A section for use with the movie scene level visibility track, which controls streamed level visibility.
  */
@@ -23,6 +27,7 @@ class MOVIESCENETRACKS_API UMovieSceneLevelVisibilitySection
 	: public UMovieSceneSection
 {
 	GENERATED_UCLASS_BODY()
+
 public:
 
 	ELevelVisibility GetVisibility() const;
@@ -30,16 +35,20 @@ public:
 
 	TArray<FName>* GetLevelNames();
 
-	// UMovieSceneSection interface
-	virtual TOptional<float> GetKeyTime( FKeyHandle KeyHandle ) const override { return TOptional<float>(); }
-	virtual void SetKeyTime( FKeyHandle KeyHandle, float Time ) override { }
+public:
+
+	//~ UMovieSceneSection interface
+
+	virtual TOptional<float> GetKeyTime(FKeyHandle KeyHandle) const override;
+	virtual void SetKeyTime(FKeyHandle KeyHandle, float Time) override;
 
 private:
+
 	/** Whether or not the levels in this section should be visible or hidden. */
-	UPROPERTY( EditAnywhere, Category = LevelVisibility )
+	UPROPERTY(EditAnywhere, Category = LevelVisibility)
 	TEnumAsByte<ELevelVisibility> Visibility;
 
 	/** The short names of the levels who's visibility is controlled by this section. */
-	UPROPERTY( EditAnywhere, Category = LevelVisibility )
+	UPROPERTY(EditAnywhere, Category = LevelVisibility)
 	TArray<FName> LevelNames;
-};
\ No newline at end of file
+};
diff --git a/Engine/Source/Runtime/MovieSceneTracks/Public/Sections/MovieSceneParameterSection.h b/Engine/Source/Runtime/MovieSceneTracks/Public/Sections/MovieSceneParameterSection.h
--- a/Engine/Source/Runtime/MovieSceneTracks/Public/Sections/MovieSceneParameterSection.h
+++ b/Engine/Source/Runtime/MovieSceneTracks/Public/Sections/MovieSceneParameterSection.h
@@ -7,12 +7,12 @@
 
 
 /**
-* Structure representing the animated value of a scalar parameter.
-*/
+ * Structure representing the animated value of a scalar parameter.
+ */
 struct FScalarParameterNameAndValue
 {
 	/** Creates a new FScalarParameterAndValue with a parameter name and a value. */
-	FScalarParameterNameAndValue( FName InParameterName, float InValue )
+	FScalarParameterNameAndValue(FName InParameterName, float InValue)
 	{
 		ParameterName = InParameterName;
 		Value = InValue;
@@ -27,12 +27,12 @@ struct FScalarParameterNameAndValue
 
 
 /**
-* Structure representing the animated value of a vector parameter.
-*/
+ * Structure representing the animated value of a vector parameter.
+ */
 struct FVectorParameterNameAndValue
 {
 	/** Creates a new FVectorParameterAndValue with a parameter name and a value. */
-	FVectorParameterNameAndValue( FName InParameterName, FVector InValue )
+	FVectorParameterNameAndValue(FName InParameterName, FVector InValue)
 	{
 		ParameterName = InParameterName;
 		Value = InValue;
@@ -47,12 +47,12 @@ struct FVectorParameterNameAndValue
 
 
 /**
-* Structure representing the animated value of a color parameter.
-*/
+ * Structure representing the animated value of a color parameter.
+ */
 struct FColorParameterNameAndValue
 {
 	/** Creates a new FColorParameterAndValue with a parameter name and a value. */
-	FColorParameterNameAndValue( FName InParameterName, FLinearColor InValue )
+	FColorParameterNameAndValue(FName InParameterName, FLinearColor InValue)
 	{
 		ParameterName = InParameterName;
 		Value = InValue;
@@ -80,7 +80,7 @@ struct FScalarParameterNameAndCurve
 	}
 
 	/** Creates a new FScalarParameterNameAndCurve for a specific scalar parameter. */
-	FScalarParameterNameAndCurve( FName InParameterName );
+	FScalarParameterNameAndCurve(FName InParameterName);
 
 	/** The name of the scalar parameter which is being animated. */
 	UPROPERTY()
@@ -96,8 +96,8 @@ struct FScalarParameterNameAndCurve
 
 
 /**
-* Structure representing an animated vector parameter and it's associated animation curve.
-*/
+ * Structure representing an animated vector parameter and it's associated animation curve.
+ */
 USTRUCT()
 struct FVectorParameterNameAndCurves
 {
@@ -109,7 +109,7 @@ struct FVectorParameterNameAndCurves
 	}
 
 	/** Creates a new FVectorParameterNameAndCurve for a specific vector parameter. */
-	FVectorParameterNameAndCurves( FName InParameterName );
+	FVectorParameterNameAndCurves(FName InParameterName);
 
 	/** The name of the vector parameter which is being animated. */
 	UPROPERTY()
@@ -146,7 +146,7 @@ struct FColorParameterNameAndCurves
 	}
 
 	/** Creates a new FVectorParameterNameAndCurve for a specific color parameter. */
-	FColorParameterNameAndCurves( FName InParameterName );
+	FColorParameterNameAndCurves(FName InParameterName);
 
 	/** The name of the vector parameter which is being animated. */
 	UPROPERTY()
@@ -176,8 +176,9 @@ struct FColorParameterNameAndCurves
 /**
  * A single movie scene section which can contain data for multiple named parameters.
  */
-UCLASS( MinimalAPI )
-class UMovieSceneParameterSection : public UMovieSceneSection
+UCLASS(MinimalAPI)
+class UMovieSceneParameterSection
+	: public UMovieSceneSection
 {
 	GENERATED_UCLASS_BODY()
 
@@ -186,10 +187,10 @@ public:
 	void AddScalarParameterKey(FName InParameterName, float InTime, float InValue);
 
 	/** Adds a a key for a specific vector parameter at the specified time with the specified value. */
-	void AddVectorParameterKey( FName InParameterName, float InTime, FVector InValue );
+	void AddVectorParameterKey(FName InParameterName, float InTime, FVector InValue);
 
 	/** Adds a a key for a specific color parameter at the specified time with the specified value. */
-	void AddColorParameterKey( FName InParameterName, float InTime, FLinearColor InValue );
+	void AddColorParameterKey(FName InParameterName, float InTime, FLinearColor InValue);
 
 	/** 
 	 * Removes a scalar parameter from this section. 
@@ -197,23 +198,23 @@ public:
 	 * @param InParameterName The name of the scalar parameter to remove.
 	 * @returns True if a parameter with that name was found and removed, otherwise false.
 	 */
-	MOVIESCENETRACKS_API bool RemoveScalarParameter( FName InParameterName );
+	MOVIESCENETRACKS_API bool RemoveScalarParameter(FName InParameterName);
 
 	/**
-	* Removes a vector parameter from this section.
-	*
-	* @param InParameterName The name of the vector parameter to remove.
-	* @returns True if a parameter with that name was found and removed, otherwise false.
-	*/
-	MOVIESCENETRACKS_API bool RemoveVectorParameter( FName InParameterName );
+	 * Removes a vector parameter from this section.
+	 *
+	 * @param InParameterName The name of the vector parameter to remove.
+	 * @returns True if a parameter with that name was found and removed, otherwise false.
+	 */
+	MOVIESCENETRACKS_API bool RemoveVectorParameter(FName InParameterName);
 
 	/**
-	* Removes a color parameter from this section.
-	*
-	* @param InParameterName The name of the color parameter to remove.
-	* @returns True if a parameter with that name was found and removed, otherwise false.
-	*/
-	MOVIESCENETRACKS_API bool RemoveColorParameter( FName InParameterName );
+	 * Removes a color parameter from this section.
+	 *
+	 * @param InParameterName The name of the color parameter to remove.
+	 * @returns True if a parameter with that name was found and removed, otherwise false.
+	 */
+	MOVIESCENETRACKS_API bool RemoveColorParameter(FName InParameterName);
 
 	/** Gets the animated scalar parameters and their associated curves. */
 	MOVIESCENETRACKS_API TArray<FScalarParameterNameAndCurve>* GetScalarParameterNamesAndCurves();
@@ -228,48 +229,45 @@ public:
 	MOVIESCENETRACKS_API void GetParameterNames(TSet<FName>& ParameterNames) const;
 
 	/**
-	* Gets the animated values for this track.
-	* @param Position The playback position to use for evaluation.
-	* @param OutScalarValues An array of FScalarParameterNameAndValue objects representing each animated scalar parameter and it's animated value.
-	* @param OutVectorValues An array of FVectorParameterNameAndValue objects representing each animated vector parameter and it's animated value.
-	* @param OutColorValues An array of FColorParameterNameAndValue objects representing each animated color parameter and it's animated value.
-	*/
-	void Eval( float Position, 
+	 * Gets the animated values for this track.
+	 *
+	 * @param Position The playback position to use for evaluation.
+	 * @param OutScalarValues An array of FScalarParameterNameAndValue objects representing each animated scalar parameter and it's animated value.
+	 * @param OutVectorValues An array of FVectorParameterNameAndValue objects representing each animated vector parameter and it's animated value.
+	 * @param OutColorValues An array of FColorParameterNameAndValue objects representing each animated color parameter and it's animated value.
+	 */
+	void Eval(float Position, 
 		TArray<FScalarParameterNameAndValue>& OutScalarValues,
 		TArray<FVectorParameterNameAndValue>& OutVectorValues,
-		TArray<FColorParameterNameAndValue>& OutColorValues ) const;
+		TArray<FColorParameterNameAndValue>& OutColorValues) const;
 
 public:
 
-	// UMovieSceneSection interface
+	//~ UMovieSceneSection interface
 
 	virtual void DilateSection(float DilationFactor, float Origin, TSet<FKeyHandle>& KeyHandles) override;
 	virtual void GetKeyHandles(TSet<FKeyHandle>& OutKeyHandles, TRange<float> TimeRange) const override;
 	virtual void MoveSection(float DeltaPosition, TSet<FKeyHandle>& KeyHandles) override;
-	virtual TOptional<float> GetKeyTime( FKeyHandle KeyHandle ) const override;
-	virtual void SetKeyTime( FKeyHandle KeyHandle, float Time ) override;
+	virtual TOptional<float> GetKeyTime(FKeyHandle KeyHandle) const override;
+	virtual void SetKeyTime(FKeyHandle KeyHandle, float Time) override;
 
 private:
+
 	void UpdateParameterIndicesFromRemoval(int32 RemovedIndex);
 	void GatherCurves(TArray<const FRichCurve*> &OutCurves) const;
 	void GatherCurves(TArray<FRichCurve*> &OutCurves);
 
 private:
-	/**
-	 * The scalar parameter names and their associated curves.
-	 */
+
+	/** The scalar parameter names and their associated curves. */
 	UPROPERTY()
 	TArray<FScalarParameterNameAndCurve> ScalarParameterNamesAndCurves;
 
-	/**
-	* The vector parameter names and their associated curves.
-	*/
+	/** The vector parameter names and their associated curves. */
 	UPROPERTY()
 	TArray<FVectorParameterNameAndCurves> VectorParameterNamesAndCurves;
 
-	/**
-	* The vector parameter names and their associated curves.
-	*/
+	/** The vector parameter names and their associated curves. */
 	UPROPERTY()
 	TArray<FColorParameterNameAndCurves> ColorParameterNamesAndCurves;
-};
\ No newline at end of file
+};
diff --git a/Engine/Source/Runtime/MovieSceneTracks/Public/Sections/MovieSceneParticleSection.h b/Engine/Source/Runtime/MovieSceneTracks/Public/Sections/MovieSceneParticleSection.h
--- a/Engine/Source/Runtime/MovieSceneTracks/Public/Sections/MovieSceneParticleSection.h
+++ b/Engine/Source/Runtime/MovieSceneTracks/Public/Sections/MovieSceneParticleSection.h
@@ -2,13 +2,14 @@
 
 #pragma once
 
+#include "Curves/IntegralCurve.h"
 #include "MovieSceneSection.h"
 #include "MovieSceneParticleSection.generated.h"
 
 
 /**
-* Defines the types of particle keys.
-*/
+ * Defines the types of particle keys.
+ */
 UENUM()
 namespace EParticleKey
 {
@@ -34,16 +35,18 @@ class UMovieSceneParticleSection
 
 	MOVIESCENETRACKS_API FIntegralCurve& GetParticleCurve();
 
-	/**
-	* UMovieSceneSection interface
-	*/
-	virtual void MoveSection( float DeltaPosition, TSet<FKeyHandle>& KeyHandles ) override;
-	virtual void DilateSection( float DilationFactor, float Origin, TSet<FKeyHandle>& KeyHandles ) override;
-	virtual void GetKeyHandles( TSet<FKeyHandle>& OutKeyHandles, TRange<float> TimeRange ) const override;
-	virtual TOptional<float> GetKeyTime( FKeyHandle KeyHandle ) const override;
-	virtual void SetKeyTime( FKeyHandle KeyHandle, float Time ) override;
+public:
+
+	//~ UMovieSceneSection interface
+
+	virtual void MoveSection(float DeltaPosition, TSet<FKeyHandle>& KeyHandles) override;
+	virtual void DilateSection(float DilationFactor, float Origin, TSet<FKeyHandle>& KeyHandles) override;
+	virtual void GetKeyHandles(TSet<FKeyHandle>& OutKeyHandles, TRange<float> TimeRange) const override;
+	virtual TOptional<float> GetKeyTime(FKeyHandle KeyHandle) const override;
+	virtual void SetKeyTime(FKeyHandle KeyHandle, float Time) override;
 
 private:
+
 	/** Curve containing the particle keys. */
 	UPROPERTY()
 	FIntegralCurve ParticleKeys;
diff --git a/Engine/Source/Runtime/MovieSceneTracks/Public/Sections/MovieSceneSkeletalAnimationSection.h b/Engine/Source/Runtime/MovieSceneTracks/Public/Sections/MovieSceneSkeletalAnimationSection.h
--- a/Engine/Source/Runtime/MovieSceneTracks/Public/Sections/MovieSceneSkeletalAnimationSection.h
+++ b/Engine/Source/Runtime/MovieSceneTracks/Public/Sections/MovieSceneSkeletalAnimationSection.h
@@ -6,14 +6,18 @@
 #include "Animation/AnimSequence.h"
 #include "MovieSceneSkeletalAnimationSection.generated.h"
 
+
 /**
  * Movie scene section that control skeletal animation
  */
 UCLASS( MinimalAPI )
-class UMovieSceneSkeletalAnimationSection : public UMovieSceneSection
+class UMovieSceneSkeletalAnimationSection
+	: public UMovieSceneSection
 {
 	GENERATED_UCLASS_BODY()
+
 public:
+
 	/** Sets the animation sequence for this section */
 	void SetAnimSequence(class UAnimSequence* InAnimSequence) { AnimSequence = InAnimSequence; }
 	
@@ -56,7 +60,10 @@ public:
 	/** Sets the anim BP slot name. */
 	void SetSlotName( FName InSlotName ) { SlotName = InSlotName; }
 
-	/** MovieSceneSection interface */
+public:
+
+	//~ MovieSceneSection interface
+
 	virtual void MoveSection( float DeltaPosition, TSet<FKeyHandle>& KeyHandles ) override;
 	virtual void DilateSection( float DilationFactor, float Origin, TSet<FKeyHandle>& KeyHandles  ) override;
 	virtual UMovieSceneSection* SplitSection(float SplitTime) override;
@@ -67,15 +74,20 @@ public:
 
 private:
 
-	static FName DefaultSlotName;
+	//~ UObject interface
 
-	// UObject interface
 #if WITH_EDITOR
+
 	virtual void PreEditChange(UProperty* PropertyAboutToChange) override;
 	virtual void PostEditChangeProperty(FPropertyChangedEvent& PropertyChangedEvent) override;
 	float PreviousPlayRate;
+
 #endif
 
+private:
+
+	static FName DefaultSlotName;
+
 	/** The animation sequence this section has */
 	UPROPERTY(EditAnywhere, Category="Animation")
 	class UAnimSequence* AnimSequence;
diff --git a/Engine/Source/Runtime/MovieSceneTracks/Public/Sections/MovieSceneSlomoSection.h b/Engine/Source/Runtime/MovieSceneTracks/Public/Sections/MovieSceneSlomoSection.h
--- a/Engine/Source/Runtime/MovieSceneTracks/Public/Sections/MovieSceneSlomoSection.h
+++ b/Engine/Source/Runtime/MovieSceneTracks/Public/Sections/MovieSceneSlomoSection.h
@@ -7,7 +7,7 @@
 
 
 /**
- * A single floating point section
+ * A single floating point section.
  */
 UCLASS(MinimalAPI)
 class UMovieSceneSlomoSection
diff --git a/Engine/Source/Runtime/MovieSceneTracks/Public/Sections/MovieSceneStringSection.h b/Engine/Source/Runtime/MovieSceneTracks/Public/Sections/MovieSceneStringSection.h
new file mode 100644
--- /dev/null
+++ b/Engine/Source/Runtime/MovieSceneTracks/Public/Sections/MovieSceneStringSection.h
@@ -0,0 +1,65 @@
+// Copyright 1998-2016 Epic Games, Inc. All Rights Reserved.
+
+#pragma once
+
+#include "MovieSceneSection.h"
+#include "IKeyframeSection.h"
+#include "Curves/StringCurve.h"
+#include "MovieSceneStringSection.generated.h"
+
+
+/**
+ * A single string section
+ */
+UCLASS(MinimalAPI)
+class UMovieSceneStringSection
+	: public UMovieSceneSection
+	, public IKeyframeSection<FString>
+{
+	GENERATED_UCLASS_BODY()
+
+public:
+
+	/**
+	 * Evaluates this section at the specified time.
+	 *
+	 * @param Time The time at which to evaluate the section.
+	 * @param DefaultString The default value to return.
+	 */
+	virtual FString Eval(float Time, const FString& DefaultString) const;
+
+	/**
+	 * Get the section's string curve.
+	 *
+	 * @return String curve.
+	 */
+	FStringCurve& GetStringCurve()
+	{
+		return StringCurve;
+	}
+
+public:
+
+	//~ IKeyframeSection interface.
+
+	virtual bool NewKeyIsNewData( float Time, const FString& KeyData ) const;
+	virtual bool HasKeys( const FString& KeyData ) const;
+	virtual void AddKey( float Time, const FString& KeyData, EMovieSceneKeyInterpolation KeyInterpolation );
+	virtual void SetDefault( const FString& KeyData );
+
+public:
+
+	//~ UMovieSceneSection interface 
+
+	virtual void DilateSection(float DilationFactor, float Origin, TSet<FKeyHandle>& KeyHandles) override;
+	virtual void GetKeyHandles(TSet<FKeyHandle>& OutKeyHandles, TRange<float> TimeRange) const override;
+	virtual TOptional<float> GetKeyTime(FKeyHandle KeyHandle) const override;
+	virtual void MoveSection(float DeltaPosition, TSet<FKeyHandle>& KeyHandles) override;
+	virtual void SetKeyTime(FKeyHandle KeyHandle, float Time) override;
+
+private:
+
+	/** Curve data */
+	UPROPERTY()
+	FStringCurve StringCurve;
+};
\ No newline at end of file
diff --git a/Engine/Source/Runtime/MovieSceneTracks/Public/Sections/MovieSceneSubSection.h b/Engine/Source/Runtime/MovieSceneTracks/Public/Sections/MovieSceneSubSection.h
--- a/Engine/Source/Runtime/MovieSceneTracks/Public/Sections/MovieSceneSubSection.h
+++ b/Engine/Source/Runtime/MovieSceneTracks/Public/Sections/MovieSceneSubSection.h
@@ -5,10 +5,12 @@
 #include "MovieSceneSection.h"
 #include "MovieSceneSubSection.generated.h"
 
+
 class UMovieSceneSequence;
 
 DECLARE_DELEGATE_OneParam(FOnSequenceChanged, UMovieSceneSequence* /*Sequence*/);
 
+
 /**
  * Implements a section in sub-sequence tracks.
  */
@@ -31,6 +33,8 @@ public:
 	 */
 	UMovieSceneSequence* GetSequence() const;
 
+public:
+
 	/**
 	 * Sets the sequence played by this section.
 	 *
@@ -88,6 +92,10 @@ public:
 		ActorToRecord = InActorToRecord;
 	}
 
+public:
+
+	//~ UMovieSceneSection interface
+
 	virtual UMovieSceneSection* SplitSection( float SplitTime ) override;
 	virtual void TrimSection( float TrimTime, bool bTrimLeft ) override;
 	virtual TOptional<float> GetKeyTime( FKeyHandle KeyHandle ) const override { return TOptional<float>(); }
diff --git a/Engine/Source/Runtime/MovieSceneTracks/Public/Sections/MovieSceneVectorSection.h b/Engine/Source/Runtime/MovieSceneTracks/Public/Sections/MovieSceneVectorSection.h
--- a/Engine/Source/Runtime/MovieSceneTracks/Public/Sections/MovieSceneVectorSection.h
+++ b/Engine/Source/Runtime/MovieSceneTracks/Public/Sections/MovieSceneVectorSection.h
@@ -2,6 +2,7 @@
 
 #pragma once
 
+#include "Curves/RichCurve.h"
 #include "IKeyframeSection.h"
 #include "MovieSceneKeyStruct.h"
 #include "MovieSceneSection.h"
@@ -49,7 +50,7 @@ struct FMovieSceneVectorKeyStruct
 
 
 /**
- * A vector section
+ * A vector section.
  */
 UCLASS(MinimalAPI)
 class UMovieSceneVectorSection
@@ -82,7 +83,7 @@ public:
 
 public:
 
-	// UMovieSceneSection interface
+	//~ UMovieSceneSection interface
 
 	virtual void MoveSection(float DeltaPosition, TSet<FKeyHandle>& KeyHandles) override;
 	virtual void DilateSection(float DilationFactor, float Origin, TSet<FKeyHandle>& KeyHandles) override;
@@ -93,7 +94,7 @@ public:
 
 public:
 
-	// IKeyframeSection interface.
+	//~ IKeyframeSection interface.
 
 	virtual void AddKey(float Time, const FVectorKey& Key, EMovieSceneKeyInterpolation KeyInterpolation) override;
 	virtual bool NewKeyIsNewData(float Time, const FVectorKey& Key) const override;
diff --git a/Engine/Source/Runtime/MovieSceneTracks/Public/TrackInstances/MovieSceneStringTrackInstance.h b/Engine/Source/Runtime/MovieSceneTracks/Public/TrackInstances/MovieSceneStringTrackInstance.h
new file mode 100644
--- /dev/null
+++ b/Engine/Source/Runtime/MovieSceneTracks/Public/TrackInstances/MovieSceneStringTrackInstance.h
@@ -0,0 +1,49 @@
+// Copyright 1998-2016 Epic Games, Inc. All Rights Reserved.
+
+#pragma once
+
+#include "IMovieSceneTrackInstance.h"
+
+
+class UMovieSceneStringTrack;
+
+
+/**
+ * Instance of a UMovieSceneStringTrack
+ */
+class FMovieSceneStringTrackInstance
+	: public IMovieSceneTrackInstance
+{
+public:
+
+	/**
+	 * Creates and initializes a new instance.
+	 *
+	 * @param InStringTrack The string track to create an instance for.
+	 */
+	FMovieSceneStringTrackInstance(UMovieSceneStringTrack& InStringTrack );
+
+	/** Virtual destructor. */
+	virtual ~FMovieSceneStringTrackInstance() { }
+
+public:
+
+	//~ IMovieSceneTrackInstance interface
+
+	virtual void ClearInstance(IMovieScenePlayer& Player, FMovieSceneSequenceInstance& SequenceInstance) override { }
+	virtual void RefreshInstance(const TArray<TWeakObjectPtr<UObject>>& RuntimeObjects, IMovieScenePlayer& Player, FMovieSceneSequenceInstance& SequenceInstance) override;
+	virtual void RestoreState(const TArray<TWeakObjectPtr<UObject>>& RuntimeObjects, IMovieScenePlayer& Player, FMovieSceneSequenceInstance& SequenceInstance) override;
+	virtual void SaveState(const TArray<TWeakObjectPtr<UObject>>& RuntimeObjects, IMovieScenePlayer& Player, FMovieSceneSequenceInstance& SequenceInstance) override;
+	virtual void Update(EMovieSceneUpdateData& UpdateData, const TArray<TWeakObjectPtr<UObject>>& RuntimeObjects, IMovieScenePlayer& Player, FMovieSceneSequenceInstance& SequenceInstance) override;
+
+private:
+
+	/** Map from object to initial state. */
+	TMap<TWeakObjectPtr<UObject>, FString> InitStringMap;
+
+	/** Runtime property bindings */
+	TSharedPtr<FTrackInstancePropertyBindings> PropertyBindings;
+
+	/** Track that is being instanced */
+	UMovieSceneStringTrack* StringTrack;
+};
