Copying //UE4/Dev-Sequencer to //UE4/Dev-Main (Source: //UE4/Dev-Sequencer @ 3237992)

diff --git a/Engine/Source/Runtime/Media/Public/IMediaOutput.h b/Engine/Source/Runtime/Media/Public/IMediaOutput.h
--- a/Engine/Source/Runtime/Media/Public/IMediaOutput.h
+++ b/Engine/Source/Runtime/Media/Public/IMediaOutput.h
@@ -5,6 +5,7 @@
 #include "CoreMinimal.h"
 
 class IMediaAudioSink;
+class IMediaBinarySink;
 class IMediaOverlaySink;
 class IMediaTextureSink;
 
@@ -21,15 +22,23 @@ public:
 	 * Set the sink to receive data from audio tracks.
 	 *
 	 * @param Sink The sink to set, or nullptr to clear.
-	 * @see SetImageSink, SetOverlaySink, SetVideoSink
+	 * @see SetBinarySink, SetOverlaySink, SetVideoSink
 	 */
 	virtual void SetAudioSink(IMediaAudioSink* Sink) = 0;
 
+	/**
+	 * Set the sink to receive data from metadata tracks.
+	 *
+	 * @param Sink The sink to set, or nullptr to clear.
+	 * @see SetAudioSink, SetOverlaySink, SetVideoSink
+	 */
+	virtual void SetMetadataSink(IMediaBinarySink* Sink) = 0;
+
 	/**
 	 * Set the sink to receive data from caption, subtitle, and text tracks.
 	 *
 	 * @param Sink The sink to set, or nullptr to clear.
-	 * @see SetAudioSink, SetImageSink, SetVideoSink
+	 * @see SetAudioSink, SetBinarySink, SetVideoSink
 	 */
 	virtual void SetOverlaySink(IMediaOverlaySink* Sink) = 0;
 
@@ -37,7 +46,7 @@ public:
 	 * Set the sink to receive data from the active video track.
 	 *
 	 * @param Sink The sink to set, or nullptr to clear.
-	 * @see SetAudioSink, SetImageSink, SetOverlaySink
+	 * @see SetAudioSink, SetBinarySink, SetOverlaySink
 	 */
 	virtual void SetVideoSink(IMediaTextureSink* Sink) = 0;
 
diff --git a/Engine/Source/Runtime/Media/Public/IMediaPlayer.h b/Engine/Source/Runtime/Media/Public/IMediaPlayer.h
--- a/Engine/Source/Runtime/Media/Public/IMediaPlayer.h
+++ b/Engine/Source/Runtime/Media/Public/IMediaPlayer.h
@@ -154,6 +154,22 @@ public:
 	 */
 	virtual bool Open(const TSharedRef<FArchive, ESPMode::ThreadSafe>& Archive, const FString& OriginalUrl, const IMediaOptions& Options) = 0;
 
+	/**
+	 * Tick the media player logic.
+	 *
+	 * @param DeltaTime Time since last tick.
+	 * @see TickVideo
+	 */
+	virtual void TickPlayer(float DeltaTime) = 0;
+
+	/**
+	 * Tick the media player's video code.
+	 *
+	 * @param DeltaTime Time since last tick.
+	 * @see TickPlayer
+	 */
+	virtual void TickVideo(float DeltaTime) = 0;
+
 public:
 
 	/** Get an event delegate that is invoked when an event occurred. */
diff --git a/Engine/Source/Runtime/Media/Public/IMediaTextureSink.h b/Engine/Source/Runtime/Media/Public/IMediaTextureSink.h
--- a/Engine/Source/Runtime/Media/Public/IMediaTextureSink.h
+++ b/Engine/Source/Runtime/Media/Public/IMediaTextureSink.h
@@ -44,7 +44,7 @@ enum class EMediaTextureSinkFormat
 	FloatRGB,
 
 	/** Four floating point components (Red, Green, Blue, Alpha) per pixel. */
-	FloatRGBA,
+	FloatRGBA
 };
 
 
@@ -126,7 +126,7 @@ public:
 	/**
 	 * Get the sink's current texture format.
 	 *
-	 * @return texture format.
+	 * @return Texture format.
 	 * @see GetTextureSinkDimensions, GetTextureSinkMode, InitializeTextureSink
 	 */
 	virtual EMediaTextureSinkFormat GetTextureSinkFormat() const = 0;
@@ -287,37 +287,6 @@ public:
 
 #endif //WITH_ENGINE
 
-public:
-
-	DEPRECATED(4.15, "This overload of InitializeTextureSink has been deprecated. Please use the other one!")
-	bool InitializeTextureSink(FIntPoint Dimensions, EMediaTextureSinkFormat Format, EMediaTextureSinkMode Mode)
-	{
-		FIntPoint BufferDim;
-
-		switch (Format)
-		{
-		case EMediaTextureSinkFormat::CharNV12:
-		case EMediaTextureSinkFormat::CharNV21:
-			BufferDim = FIntPoint(Align(Dimensions.X, 16), Align(Dimensions.Y, 2) * 3 / 2);
-			break;
-
-		case EMediaTextureSinkFormat::CharUYVY:
-		case EMediaTextureSinkFormat::CharYUY2:
-		case EMediaTextureSinkFormat::CharYVYU:
-			BufferDim = FIntPoint(Align(Dimensions.X, 16) / 2, Dimensions.Y);
-			break;
-
-		case EMediaTextureSinkFormat::CharAYUV:
-		case EMediaTextureSinkFormat::CharBGRA:
-		case EMediaTextureSinkFormat::FloatRGB:
-		case EMediaTextureSinkFormat::FloatRGBA:
-			BufferDim = Dimensions;
-			break;
-		}
-
-		return InitializeTextureSink(Dimensions, BufferDim, Format, Mode);
-	}
-
 public:
 
 	/** Virtual destructor. */
diff --git a/Engine/Source/Runtime/Media/Public/IMediaTracks.h b/Engine/Source/Runtime/Media/Public/IMediaTracks.h
--- a/Engine/Source/Runtime/Media/Public/IMediaTracks.h
+++ b/Engine/Source/Runtime/Media/Public/IMediaTracks.h
@@ -12,12 +12,12 @@ enum class EMediaTrackType
 	/** Audio track. */
 	Audio,
 
-	/** Binary data track. */
-	Binary,
-
 	/** Closed caption track. */
 	Caption,
 
+	/** Metadata track. */
+	Metadata,
+
 	/** Script track. */
 	Script,
 
diff --git a/Engine/Source/Runtime/MediaAssets/Private/Assets/BaseMediaSource.cpp b/Engine/Source/Runtime/MediaAssets/Private/Assets/BaseMediaSource.cpp
new file mode 100644
--- /dev/null
+++ b/Engine/Source/Runtime/MediaAssets/Private/Assets/BaseMediaSource.cpp
@@ -0,0 +1,94 @@
+// Copyright 1998-2017 Epic Games, Inc. All Rights Reserved.
+
+#include "BaseMediaSource.h"
+#include "UObject/SequencerObjectVersion.h"
+
+#if WITH_EDITOR
+	#include "Interfaces/ITargetPlatform.h"
+#endif
+
+
+/* UObject interface
+ *****************************************************************************/
+
+void UBaseMediaSource::GetAssetRegistryTags(TArray<FAssetRegistryTag>& OutTags) const
+{
+	FString Url = GetUrl();
+
+	if (!Url.IsEmpty())
+	{
+		OutTags.Add(FAssetRegistryTag("Url", Url, FAssetRegistryTag::TT_Alphabetical));
+	}
+}
+
+
+#if WITH_EDITOR
+void UBaseMediaSource::GetAssetRegistryTagMetadata(TMap<FName, FAssetRegistryTagMetadata>& OutMetadata) const
+{
+}
+#endif
+
+
+void UBaseMediaSource::Serialize(FArchive& Ar)
+{
+	Super::Serialize(Ar);
+
+	Ar.UsingCustomVersion(FSequencerObjectVersion::GUID);
+	auto CustomVersion = Ar.CustomVer(FSequencerObjectVersion::GUID);
+
+	if (Ar.IsLoading() && (CustomVersion < FSequencerObjectVersion::RenameMediaSourcePlatformPlayers))
+	{
+#if WITH_EDITORONLY_DATA
+		if (!Ar.IsFilterEditorOnly())
+		{
+			TMap<FString, FString> DummyPlatformPlayers;
+			Ar << DummyPlatformPlayers;
+		}
+#endif
+
+		FString DummyDefaultPlayer;
+		Ar << DummyDefaultPlayer;
+	}
+	else
+	{
+#if WITH_EDITORONLY_DATA
+		if (Ar.IsFilterEditorOnly())
+		{
+			if (Ar.IsSaving())
+			{
+				const FName* PlatformPlayerName = PlatformPlayerNames.Find(Ar.CookingTarget()->IniPlatformName());
+				PlayerName = (PlatformPlayerName != nullptr) ? *PlatformPlayerName : NAME_None;
+			}
+
+			Ar << PlayerName;
+		}
+		else
+		{
+			Ar << PlatformPlayerNames;
+		}
+#else
+		Ar << PlayerName;
+#endif
+	}
+}
+
+
+/* IMediaOptions interface
+ *****************************************************************************/
+
+FName UBaseMediaSource::GetDesiredPlayerName() const
+{
+#if WITH_EDITORONLY_DATA
+	const FString RunningPlatformName(FPlatformProperties::IniPlatformName());
+	const FName* PlatformPlayerName = PlatformPlayerNames.Find(RunningPlatformName);
+
+	if (PlatformPlayerName == nullptr)
+	{
+		return Super::GetDesiredPlayerName();
+	}
+
+	return *PlatformPlayerName;
+#else
+	return PlayerName;
+#endif
+}
diff --git a/Engine/Source/Runtime/MediaAssets/Private/Assets/FileMediaSource.cpp b/Engine/Source/Runtime/MediaAssets/Private/Assets/FileMediaSource.cpp
--- a/Engine/Source/Runtime/MediaAssets/Private/Assets/FileMediaSource.cpp
+++ b/Engine/Source/Runtime/MediaAssets/Private/Assets/FileMediaSource.cpp
@@ -4,7 +4,11 @@
 #include "Misc/Paths.h"
 
 
-static FName PrecacheFileName("PrecacheFile");
+namespace FileMediaSource
+{
+	/** Name of the PrecacheFile media option. */
+	static const FName PrecacheFileOption("PrecacheFile");
+}
 
 
 /* UFileMediaSource interface
@@ -32,20 +36,34 @@ void UFileMediaSource::SetFilePath(const FString& Path)
 }
 
 
-/* UMediaSource overrides
+/* IMediaSource overrides
  *****************************************************************************/
 
 bool UFileMediaSource::GetMediaOption(const FName& Key, bool DefaultValue) const
 {
-	if (Key == PrecacheFileName)
+	if (Key == FileMediaSource::PrecacheFileOption)
 	{
 		return PrecacheFile;
 	}
 
-	return DefaultValue;
+	return Super::GetMediaOption(Key, DefaultValue);
 }
 
 
+bool UFileMediaSource::HasMediaOption(const FName& Key) const
+{
+	if (Key == FileMediaSource::PrecacheFileOption)
+	{
+		return true;
+	}
+
+	return Super::HasMediaOption(Key);
+}
+
+
+/* UMediaSource overrides
+ *****************************************************************************/
+
 FString UFileMediaSource::GetUrl() const
 {
 	return FString(TEXT("file://")) + GetFullPath();
diff --git a/Engine/Source/Runtime/MediaAssets/Private/Assets/MediaOverlays.cpp b/Engine/Source/Runtime/MediaAssets/Private/Assets/MediaOverlays.cpp
new file mode 100644
--- /dev/null
+++ b/Engine/Source/Runtime/MediaAssets/Private/Assets/MediaOverlays.cpp
@@ -0,0 +1,100 @@
+// Copyright 1998-2016 Epic Games, Inc. All Rights Reserved.
+
+#include "MediaOverlays.h"
+
+
+/* UMediaPlayer interface
+*****************************************************************************/
+
+void UMediaOverlays::GetItems(EMediaOverlayType Type, TArray<FMediaPlayerOverlay>& OutItems, FTimespan Time) const
+{
+	if (Time.IsZero())
+	{
+		return;
+	}
+
+	// @todo gmp: use more efficient lookup (range tree?)
+	for (const FOverlay& Overlay : Overlays)
+	{
+		if ((Overlay.Type == Type) && (Overlay.Time <= Time) && (Time - Overlay.Time < Overlay.Duration))
+		{
+			const int32 ItemIndex = OutItems.AddDefaulted();
+			FMediaPlayerOverlay& OutItem = OutItems[ItemIndex];
+			{
+				OutItem.HasPosition = Overlay.Position.IsSet();
+				OutItem.Position = OutItem.HasPosition ? Overlay.Position.GetValue() : FIntPoint::ZeroValue;
+				OutItem.Text = Overlay.Text;
+			}
+		}
+	}
+}
+
+
+/* UObject interface
+ *****************************************************************************/
+
+void UMediaOverlays::BeginDestroy()
+{
+	Super::BeginDestroy();
+
+	BeginDestroyEvent.Broadcast(*this);
+}
+
+
+/* IMediaOverlaySink interface
+*****************************************************************************/
+
+bool UMediaOverlays::InitializeOverlaySink()
+{
+	return true;
+}
+
+
+void UMediaOverlays::AddOverlaySinkText(const FText& Text, EMediaOverlayType Type, FTimespan Time, FTimespan Duration, TOptional<FVector2D> Position)
+{
+	// @todo gmp: prevent duplicates
+	// @todo gmp: make thread-safe
+/*
+	// remove expired overlays
+	const FTimespan CurrentTime = NativePlayer.IsValid()
+		? NativePlayer->GetControls().GetTime()
+		: FTimespan::Zero();
+
+	for (int32 OverlayIndex = Overlays.Num() - 1; OverlayIndex >= 0; --OverlayIndex)
+	{
+		const FOverlay& Overlay = Overlays[OverlayIndex];
+		const FTimespan Offset = CurrentTime - Overlay.Time;
+
+		if (Offset > Overlay.Duration)
+		{
+			Overlays.RemoveAtSwap(OverlayIndex);
+		}
+	}
+*/
+	// add new overlay
+	FOverlay Overlay;
+	{
+		Overlay.Duration = Duration;
+		Overlay.Position = Position;
+		Overlay.Text = Text;
+		Overlay.Time = Time;
+		Overlay.Type = Type;
+	}
+
+	Overlays.Add(Overlay);
+}
+
+
+void UMediaOverlays::ClearOverlaySinkText()
+{
+	// @todo gmp: make thread-safe
+
+	Overlays.Empty();
+}
+
+
+void UMediaOverlays::ShutdownOverlaySink()
+{
+	// @todo gmp: make thread-safe
+	Overlays.Empty();
+}
diff --git a/Engine/Source/Runtime/MediaAssets/Private/Assets/MediaPlayer.cpp b/Engine/Source/Runtime/MediaAssets/Private/Assets/MediaPlayer.cpp
--- a/Engine/Source/Runtime/MediaAssets/Private/Assets/MediaPlayer.cpp
+++ b/Engine/Source/Runtime/MediaAssets/Private/Assets/MediaPlayer.cpp
@@ -1,19 +1,21 @@
 // Copyright 1998-2017 Epic Games, Inc. All Rights Reserved.
 
 #include "MediaPlayer.h"
-#include "Modules/ModuleManager.h"
+
 #include "IMediaControls.h"
 #include "IMediaModule.h"
 #include "IMediaOutput.h"
 #include "IMediaPlayer.h"
 #include "IMediaPlayerFactory.h"
 #include "IMediaTracks.h"
+#include "MediaOverlays.h"
 #include "Misc/Paths.h"
 #include "MediaAssetsPrivate.h"
 #include "MediaPlaylist.h"
 #include "MediaSource.h"
 #include "MediaSoundWave.h"
 #include "MediaTexture.h"
+#include "Modules/ModuleManager.h"
 
 
 /* UMediaPlayer structors
@@ -28,7 +30,9 @@ UMediaPlayer::UMediaPlayer(const FObjectInitializer& ObjectInitializer)
 #if WITH_EDITOR
 	, WasPlayingInPIE(false)
 #endif
-{ }
+{
+	Player.OnMediaEvent().AddUObject(this, &UMediaPlayer::HandlePlayerMediaEvent);
+}
 
 
 /* UMediaPlayer interface
@@ -36,14 +40,7 @@ UMediaPlayer::UMediaPlayer(const FObjectInitializer& ObjectInitializer)
 
 bool UMediaPlayer::CanPause() const
 {
-	if (!Player.IsValid())
-	{
-		return false;
-	}
-	
-	EMediaState State = Player->GetControls().GetState();
-
-	return ((State== EMediaState::Playing) || (State== EMediaState::Preparing));
+	return Player.CanPause();
 }
 
 
@@ -54,7 +51,7 @@ bool UMediaPlayer::CanPlaySource(UMediaSource* MediaSource)
 		return false;
 	}
 
-	return FindPlayerForUrl(MediaSource->GetUrl(), *MediaSource).IsValid();
+	return Player.CanPlayUrl(MediaSource->GetUrl(), *MediaSource);
 }
 
 
@@ -65,145 +62,125 @@ bool UMediaPlayer::CanPlayUrl(const FString& Url)
 		return false;
 	}
 
-	return FindPlayerForUrl(Url, *GetDefault<UMediaSource>()).IsValid();
+	return Player.CanPlayUrl(Url, *GetDefault<UMediaSource>());
 }
 
 
 void UMediaPlayer::Close()
 {
-	CurrentUrl.Empty();
-	Playlist = nullptr;
+	Player.Close();
 
-	if (Player.IsValid())
-	{
-		Player->Close();
-	}
+	LastUrl.Empty();
+	Playlist = nullptr;
+	PlaylistIndex = INDEX_NONE;
 }
 
 
-FTimespan UMediaPlayer::GetDuration() const
+FMediaPlayerBase& UMediaPlayer::GetBasePlayer()
 {
-	return Player.IsValid() ? Player->GetControls().GetDuration() : FTimespan::Zero();
+	return Player;
 }
 
 
-FFloatRange UMediaPlayer::GetForwardRates(bool Unthinned)
+FName UMediaPlayer::GetDesiredPlayerName() const
 {
-	return Player.IsValid()
-		? Player->GetControls().GetSupportedRates(EMediaPlaybackDirections::Forward, Unthinned)
-		: FFloatRange::Empty();
+	return Player.DesiredPlayerName;
 }
 
 
-int32 UMediaPlayer::GetNumTracks(EMediaPlayerTrack TrackType) const
+FTimespan UMediaPlayer::GetDuration() const
 {
-	return Player.IsValid() ? Player->GetTracks().GetNumTracks((EMediaTrackType)TrackType) : 0;
+	return Player.GetDuration();
 }
 
 
-void UMediaPlayer::GetOverlays(EMediaOverlayType Type, TArray<FMediaPlayerOverlay>& OutOverlays) const
+FFloatRange UMediaPlayer::GetForwardRates(bool Unthinned)
 {
-	FTimespan CurrentTime = GetTime();
+	return Player.GetForwardRates(Unthinned);
+}
 
-	if (CurrentTime == FTimespan::Zero())
-	{
-		return;
-	}
 
-	for (const FOverlay& Overlay : Overlays)
-	{
-		if ((Overlay.Type == Type) && (Overlay.Time <= CurrentTime) && (CurrentTime - Overlay.Time < Overlay.Duration))
-		{
-			const int32 OverlayIndex = OutOverlays.AddDefaulted();
-			FMediaPlayerOverlay& OutOverlay = OutOverlays[OverlayIndex];
-			{
-				OutOverlay.HasPosition = Overlay.Position.IsSet();
-				OutOverlay.Position = OutOverlay.HasPosition ? Overlay.Position.GetValue() : FIntPoint::ZeroValue;
-				OutOverlay.Text = Overlay.Text;
-			}
-		}
-	}
+int32 UMediaPlayer::GetNumTracks(EMediaPlayerTrack TrackType) const
+{
+	return Player.GetNumTracks((EMediaTrackType)TrackType);
 }
 
 
 FName UMediaPlayer::GetPlayerName() const
 {
-	return Player.IsValid() ? Player->GetName() : NAME_None;
+	return Player.GetPlayerName();
 }
 
 
 float UMediaPlayer::GetRate() const
 {
-	return Player.IsValid() ? Player->GetControls().GetRate() : 0.0f;
+	return Player.GetRate();
 }
 
 
 FFloatRange UMediaPlayer::GetReverseRates(bool Unthinned)
 {
-	return Player.IsValid()
-		? Player->GetControls().GetSupportedRates(EMediaPlaybackDirections::Reverse, Unthinned)
-		: FFloatRange::Empty();
+	return Player.GetReverseRates(Unthinned);
 }
 
 
 int32 UMediaPlayer::GetSelectedTrack(EMediaPlayerTrack TrackType) const
 {
-	return Player.IsValid() ? Player->GetTracks().GetSelectedTrack((EMediaTrackType)TrackType) : INDEX_NONE;
+	return Player.GetSelectedTrack((EMediaTrackType)TrackType);
 }
 
 
 FTimespan UMediaPlayer::GetTime() const
 {
-	return Player.IsValid() ? Player->GetControls().GetTime() : FTimespan::Zero();
+	return Player.GetTime();
 }
 
 
 FText UMediaPlayer::GetTrackDisplayName(EMediaPlayerTrack TrackType, int32 TrackIndex) const
 {
-	return Player.IsValid() ? Player->GetTracks().GetTrackDisplayName((EMediaTrackType)TrackType, TrackIndex) : FText::GetEmpty();
+	return Player.GetTrackDisplayName((EMediaTrackType)TrackType, TrackIndex);
 }
 
 
 FString UMediaPlayer::GetTrackLanguage(EMediaPlayerTrack TrackType, int32 TrackIndex) const
 {
-	return Player.IsValid() ? Player->GetTracks().GetTrackLanguage((EMediaTrackType)TrackType, TrackIndex) : FString();
+	return Player.GetTrackLanguage((EMediaTrackType)TrackType, TrackIndex);
+}
+
+
+const FString& UMediaPlayer::GetUrl() const
+{
+	return Player.GetUrl();
 }
 
 
 bool UMediaPlayer::IsLooping() const
 {
-	return Player.IsValid() && Player->GetControls().IsLooping();
+	return Player.IsLooping();
 }
 
 
 bool UMediaPlayer::IsPaused() const
 {
-	return Player.IsValid() && (Player->GetControls().GetState() == EMediaState::Paused);
+	return Player.IsPaused();
 }
 
 
 bool UMediaPlayer::IsPlaying() const
 {
-	return Player.IsValid() && (Player->GetControls().GetState() == EMediaState::Playing);
+	return Player.IsPlaying();
 }
 
 
 bool UMediaPlayer::IsPreparing() const
 {
-	return Player.IsValid() && (Player->GetControls().GetState() == EMediaState::Preparing);
+	return Player.IsPreparing();
 }
 
 
 bool UMediaPlayer::IsReady() const
 {
-	if (!Player.IsValid())
-	{
-		return false;
-	}
-	
-	return ((Player->GetControls().GetState() != EMediaState::Closed) &&
-			(Player->GetControls().GetState() != EMediaState::Error) &&
-			(Player->GetControls().GetState() != EMediaState::Preparing));
+	return Player.IsReady();
 }
 
 
@@ -222,7 +199,7 @@ bool UMediaPlayer::Next()
 			? Playlist->GetRandom(PlaylistIndex)
 			: Playlist->GetNext(PlaylistIndex);
 
-		if ((NextSource != nullptr) && NextSource->Validate() && Open(NextSource->GetUrl(), *NextSource))
+		if ((NextSource != nullptr) && NextSource->Validate() && Player.Open(NextSource->GetUrl(), *NextSource))
 		{
 			return true;
 		}
@@ -252,26 +229,32 @@ bool UMediaPlayer::OpenFile(const FString& FilePath)
 
 bool UMediaPlayer::OpenPlaylistIndex(UMediaPlaylist* InPlaylist, int32 Index)
 {
+	Close();
+
 	if (InPlaylist == nullptr)
 	{
+		UE_LOG(LogMediaAssets, Warning, TEXT("UMediaPlayer::OpenPlaylistIndex called with null MediaPlaylist"));
 		return false;
 	}
 
 	Playlist = InPlaylist;
-	PlaylistIndex = Index;
 
-	if (PlaylistIndex == INDEX_NONE)
+	if (Index == INDEX_NONE)
 	{
 		return true;
 	}
 
-	UMediaSource* MediaSource = Playlist->Get(PlaylistIndex);
+	UMediaSource* MediaSource = Playlist->Get(Index);
 
 	if (MediaSource == nullptr)
 	{
+		UE_LOG(LogMediaAssets, Warning, TEXT("UMediaPlayer::OpenPlaylistIndex called with invalid PlaylistIndex %i"), Index);
 		return false;
 	}
 	
+	LastUrl = MediaSource->GetUrl();
+	PlaylistIndex = Index;
+
 	if (!MediaSource->Validate())
 	{
 		UE_LOG(LogMediaAssets, Error, TEXT("Failed to validate media source %s (%s)"), *MediaSource->GetName(), *MediaSource->GetUrl());
@@ -279,17 +262,22 @@ bool UMediaPlayer::OpenPlaylistIndex(UMediaPlaylist* InPlaylist, int32 Index)
 		return false;
 	}
 
-	return Open(MediaSource->GetUrl(), *MediaSource);
+	return Player.Open(LastUrl, *MediaSource);
 }
 
 
 bool UMediaPlayer::OpenSource(UMediaSource* MediaSource)
 {
+	Close();
+
 	if (MediaSource == nullptr)
 	{
+		UE_LOG(LogMediaAssets, Warning, TEXT("UMediaPlayer::OpenSource called with null MediaSource"));
 		return false;
 	}
 
+	LastUrl = MediaSource->GetUrl();
+
 	if (!MediaSource->Validate())
 	{
 		UE_LOG(LogMediaAssets, Error, TEXT("Failed to validate media source %s (%s)"), *MediaSource->GetName(), *MediaSource->GetUrl());
@@ -297,29 +285,34 @@ bool UMediaPlayer::OpenSource(UMediaSource* MediaSource)
 		return false;
 	}
 
-	Playlist = nullptr;
-
-	return Open(MediaSource->GetUrl(), *MediaSource);
+	return Player.Open(LastUrl, *MediaSource);
 }
 
 
 bool UMediaPlayer::OpenUrl(const FString& Url)
 {
-	Playlist = nullptr;
+	Close();
 
-	return Open(Url, *GetDefault<UMediaSource>());
+	if (Url.IsEmpty())
+	{
+		return false;
+	}
+
+	LastUrl = Url;
+
+	return Player.Open(LastUrl, *GetDefault<UMediaSource>());
 }
 
 
 bool UMediaPlayer::Pause()
 {
-	return SetRate(0.0f);
+	return Player.SetRate(0.0f);
 }
 
 
 bool UMediaPlayer::Play()
 {
-	return SetRate(1.0f);
+	return Player.SetRate(1.0f);
 }
 
 
@@ -338,7 +331,7 @@ bool UMediaPlayer::Previous()
 			? Playlist->GetRandom(PlaylistIndex)
 			: Playlist->GetNext(PlaylistIndex);
 
-		if ((PrevSource != nullptr) && PrevSource->Validate() && Open(PrevSource->GetUrl(), *PrevSource))
+		if ((PrevSource != nullptr) && PrevSource->Validate() && Player.Open(PrevSource->GetUrl(), *PrevSource))
 		{
 			return true;
 		}
@@ -355,12 +348,7 @@ bool UMediaPlayer::Reopen()
 		return OpenPlaylistIndex(Playlist, PlaylistIndex);
 	}
 
-	if (!CurrentUrl.IsEmpty())
-	{
-		OpenUrl(CurrentUrl);
-	}
-
-	return false;
+	return Player.Open(LastUrl, *GetDefault<UMediaSource>());
 }
 
 
@@ -372,13 +360,19 @@ bool UMediaPlayer::Rewind()
 
 bool UMediaPlayer::Seek(const FTimespan& Time)
 {
-	return Player.IsValid() && Player->GetControls().Seek(Time);
+	return Player.Seek(Time);
 }
 
 
 bool UMediaPlayer::SelectTrack(EMediaPlayerTrack TrackType, int32 TrackIndex)
 {
-	return Player.IsValid() ? Player->GetTracks().SelectTrack((EMediaTrackType)TrackType, TrackIndex) : false;
+	return Player.SelectTrack((EMediaTrackType)TrackType, TrackIndex);
+}
+
+
+void UMediaPlayer::SetDesiredPlayerName(FName PlayerName)
+{
+	Player.DesiredPlayerName = PlayerName;
 }
 
 
@@ -386,13 +380,31 @@ bool UMediaPlayer::SetLooping(bool Looping)
 {
 	Loop = Looping;
 
-	return Player.IsValid() && Player->GetControls().SetLooping(Looping);
+	return Player.SetLooping(Looping);
+}
+
+
+void UMediaPlayer::SetOverlays(UMediaOverlays* NewOverlays)
+{
+	if (Overlays != nullptr)
+	{
+		Overlays->OnBeginDestroy().RemoveAll(this);
+	}
+
+	if (NewOverlays != nullptr)
+	{
+		NewOverlays->OnBeginDestroy().AddUObject(this, &UMediaPlayer::HandleMediaOverlaysBeginDestroy);
+	}
+
+	Player.SetOverlaySink(NewOverlays);
+
+	Overlays = NewOverlays;
 }
 
 
 bool UMediaPlayer::SetRate(float Rate)
 {
-	return Player.IsValid() && Player->GetControls().SetRate(Rate);
+	return Player.SetRate(Rate);
 }
 
 
@@ -408,10 +420,7 @@ void UMediaPlayer::SetSoundWave(UMediaSoundWave* NewSoundWave)
 		NewSoundWave->OnBeginDestroy().AddUObject(this, &UMediaPlayer::HandleMediaSoundWaveBeginDestroy);
 	}
 
-	if (Player.IsValid())
-	{
-		Player->GetOutput().SetAudioSink(NewSoundWave);
-	}
+	Player.SetAudioSink(NewSoundWave);
 
 	SoundWave = NewSoundWave;
 }
@@ -429,10 +438,7 @@ void UMediaPlayer::SetVideoTexture(UMediaTexture* NewTexture)
 		NewTexture->OnBeginDestroy().AddUObject(this, &UMediaPlayer::HandleMediaTextureBeginDestroy);
 	}
 
-	if (Player.IsValid())
-	{
-		Player->GetOutput().SetVideoSink(NewTexture);
-	}
+	Player.SetVideoSink(NewTexture);
 
 	VideoTexture = NewTexture;
 }
@@ -440,19 +446,19 @@ void UMediaPlayer::SetVideoTexture(UMediaTexture* NewTexture)
 
 bool UMediaPlayer::SupportsRate(float Rate, bool Unthinned) const
 {
-	return Player.IsValid() && Player->GetControls().SupportsRate(Rate, Unthinned);
+	return Player.SupportsRate(Rate, Unthinned);
 }
 
 
 bool UMediaPlayer::SupportsScrubbing() const
 {
-	return Player.IsValid() && Player->GetControls().SupportsScrubbing();
+	return Player.SupportsScrubbing();
 }
 
 
 bool UMediaPlayer::SupportsSeeking() const
 {
-	return Player.IsValid() && Player->GetControls().SupportsSeeking();
+	return Player.SupportsSeeking();
 }
 
 
@@ -487,27 +493,17 @@ void UMediaPlayer::BeginDestroy()
 {
 	Super::BeginDestroy();
 
+	Close();
+
+	SetOverlays(nullptr);
 	SetSoundWave(nullptr);
 	SetVideoTexture(nullptr);
-
-	if (Player.IsValid())
-	{
-		Player->Close();
-		Player->GetOutput().SetOverlaySink(nullptr);
-		Player->OnMediaEvent().RemoveAll(this);
-		Player.Reset();
-	}
 }
 
 
 FString UMediaPlayer::GetDesc()
 {
-	if (Player.IsValid())
-	{
-		return TEXT("UMediaPlayer");
-	}
-
-	return FString();
+	return TEXT("UMediaPlayer");
 }
 
 
@@ -515,6 +511,11 @@ void UMediaPlayer::PostLoad()
 {
 	Super::PostLoad();
 
+	if (Overlays != nullptr)
+	{
+		Overlays->OnBeginDestroy().AddUObject(this, &UMediaPlayer::HandleMediaOverlaysBeginDestroy);
+	}
+
 	if (SoundWave != nullptr)
 	{
 		SoundWave->OnBeginDestroy().AddUObject(this, &UMediaPlayer::HandleMediaSoundWaveBeginDestroy);
@@ -535,7 +536,11 @@ void UMediaPlayer::PostEditChangeProperty(FPropertyChangedEvent& PropertyChanged
 		? PropertyChangedEvent.Property->GetFName()
 		: NAME_None;
 
-	if (PropertyName == GET_MEMBER_NAME_CHECKED(UMediaPlayer, SoundWave))
+	if (PropertyName == GET_MEMBER_NAME_CHECKED(UMediaPlayer, Overlays))
+	{
+		SetOverlays(Overlays);
+	}
+	else if (PropertyName == GET_MEMBER_NAME_CHECKED(UMediaPlayer, SoundWave))
 	{
 		SetSoundWave(SoundWave);
 	}
@@ -558,7 +563,11 @@ void UMediaPlayer::PreEditChange(UProperty* PropertyAboutToChange)
 		? PropertyAboutToChange->GetFName()
 		: NAME_None;
 
-	if (PropertyName == GET_MEMBER_NAME_CHECKED(UMediaPlayer, SoundWave))
+	if (PropertyName == GET_MEMBER_NAME_CHECKED(UMediaPlayer, Overlays))
+	{
+		SetOverlays(nullptr);
+	}
+	else if (PropertyName == GET_MEMBER_NAME_CHECKED(UMediaPlayer, SoundWave))
 	{
 		SetSoundWave(nullptr);
 	}
@@ -573,220 +582,50 @@ void UMediaPlayer::PreEditChange(UProperty* PropertyAboutToChange)
 #endif
 
 
-/* UMediaPlayer implementation
+/* FTickerObjectBase interface
  *****************************************************************************/
 
-TSharedPtr<IMediaPlayer> UMediaPlayer::FindPlayerForUrl(const FString& Url, const IMediaOptions& Options)
-{
-	const FName PlayerName = (DesiredPlayerName != NAME_None) ? DesiredPlayerName : Options.GetDesiredPlayerName();
-
-	// reuse existing player if desired
-	if (Player.IsValid() && (PlayerName == Player->GetName()))
-	{
-		return Player;
-	}
-
-	// load media module
-	IMediaModule* MediaModule = FModuleManager::LoadModulePtr<IMediaModule>("Media");
-
-	if (MediaModule == nullptr)
-	{
-		UE_LOG(LogMediaAssets, Error, TEXT("Failed to load Media module"));
-		return nullptr;
-	}
-
-	// try to create desired player
-	if (PlayerName != NAME_None)
-	{
-		IMediaPlayerFactory* Factory = MediaModule->GetPlayerFactory(PlayerName);
-
-		if (Factory == nullptr)
-		{
-			UE_LOG(LogMediaAssets, Error, TEXT("Could not find desired player %s for %s"), *PlayerName.ToString(), *Url);
-			return nullptr;
-		}
-
-		TSharedPtr<IMediaPlayer> NewPlayer = Factory->CreatePlayer();
-
-		if (!NewPlayer.IsValid())
-		{
-			UE_LOG(LogMediaAssets, Error, TEXT("Failed to create desired player %s for %s"), *PlayerName.ToString(), *Url);
-			return nullptr;
-		}
-
-		return NewPlayer;
-	}
-
-	// try to reuse existing player
-	if (Player.IsValid())
-	{
-		IMediaPlayerFactory* Factory = MediaModule->GetPlayerFactory(Player->GetName());
-
-		if ((Factory != nullptr) && Factory->CanPlayUrl(Url, Options))
-		{
-			return Player;
-		}
-	}
-
-	// try to auto-select new player
-	const TArray<IMediaPlayerFactory*>& PlayerFactories = MediaModule->GetPlayerFactories();
-	const FString RunningPlatformName(FPlatformProperties::IniPlatformName());
-
-	for (IMediaPlayerFactory* Factory : PlayerFactories)
-	{
-		if (!Factory->SupportsPlatform(RunningPlatformName) || !Factory->CanPlayUrl(Url, Options))
-		{
-			continue;
-		}
-
-		TSharedPtr<IMediaPlayer> NewPlayer = Factory->CreatePlayer();
-
-		if (NewPlayer.IsValid())
-		{
-			return NewPlayer;
-		}
-	}
-
-	UE_LOG(LogMediaAssets, Error, TEXT("Could not find a native player for %s"), *Url);
-
-	return nullptr;
-}
-
-
-bool UMediaPlayer::Open(const FString& Url, const IMediaOptions& Options)
+bool UMediaPlayer::Tick(float DeltaTime)
 {
-	if (IsRunningDedicatedServer())
-	{
-		return false;
-	}
-
-	if (Url.IsEmpty())
-	{
-		return false;
-	}
-
-	// find & initialize new player
-	TSharedPtr<IMediaPlayer> NewPlayer = FindPlayerForUrl(Url, Options);
-
-	if (NewPlayer != Player)
-	{
-		if (Player.IsValid())
-		{
-			Player->Close();
-			Player->OnMediaEvent().RemoveAll(this);
-		}
+	Player.TickPlayer(DeltaTime);
 
-		if (NewPlayer.IsValid())
+	ENQUEUE_UNIQUE_RENDER_COMMAND_TWOPARAMETER(MediaPlayerTickRender,
+		UMediaPlayer*, This, this,
+		float, DeltaTime, DeltaTime,
 		{
-			NewPlayer->OnMediaEvent().AddUObject(this, &UMediaPlayer::HandlePlayerMediaEvent);
-		}
-
-		Player = NewPlayer;
-	}
-
-	if (!Player.IsValid())
-	{
-		return false;
-	}
-
-	CurrentUrl = Url;
-
-	// open the new media source
-	if (!Player->Open(Url, Options))
-	{
-		CurrentUrl.Empty();
-
-		return false;
-	}
+			This->Player.TickVideo(DeltaTime);
+		});
 
 	return true;
 }
 
 
-void UMediaPlayer::SelectDefaultTracks()
-{
-	IMediaTracks& Tracks = Player->GetTracks();
-
-	// @todo Media: consider locale when selecting default tracks
-	Tracks.SelectTrack(EMediaTrackType::Audio, 0);
-	Tracks.SelectTrack(EMediaTrackType::Caption, 0);
-	Tracks.SelectTrack(EMediaTrackType::Video, 0);
-}
-
-
-/* IMediaOverlaySink interface
+/* UMediaPlayer callbacks
  *****************************************************************************/
 
-bool UMediaPlayer::InitializeOverlaySink()
-{
-	return true;
-}
-
-
-void UMediaPlayer::AddOverlaySinkText(const FText& Text, EMediaOverlayType Type, FTimespan Time, FTimespan Duration, TOptional<FVector2D> Position)
+void UMediaPlayer::HandleMediaOverlaysBeginDestroy(UMediaOverlays& DestroyedOverlays)
 {
-	// @todo gmp: make thread-safe
-
-	// remove expired overlays
-	const FTimespan CurrentTime = GetTime();
-
-	for (int32 OverlayIndex = Overlays.Num() - 1; OverlayIndex >= 0; --OverlayIndex)
-	{
-		const FOverlay& Overlay = Overlays[OverlayIndex];
-		const FTimespan Offset = CurrentTime - Overlay.Time;
-
-		if (Offset > Overlay.Duration)
-		{
-			Overlays.RemoveAtSwap(OverlayIndex);
-		}
-	}
-
-	// add new overlay
-	FOverlay Overlay;
+	if (&DestroyedOverlays == Overlays)
 	{
-		Overlay.Duration = Duration;
-		Overlay.Position = Position;
-		Overlay.Text = Text;
-		Overlay.Time = Time;
-		Overlay.Type = Type;
+		Player.SetOverlaySink(nullptr);
 	}
-
-	Overlays.Add(Overlay);
-}
-
-
-void UMediaPlayer::ClearOverlaySinkText()
-{
-	// @todo gmp: make thread-safe
-
-	Overlays.Empty();
-}
-
-
-void UMediaPlayer::ShutdownOverlaySink()
-{
-	// @todo gmp: make thread-safe
-	Overlays.Empty();
 }
 
 
-/* UMediaPlayer callbacks
- *****************************************************************************/
-
 void UMediaPlayer::HandleMediaSoundWaveBeginDestroy(UMediaSoundWave& DestroyedSoundWave)
 {
-	if (Player.IsValid() && (&DestroyedSoundWave == SoundWave))
+	if (&DestroyedSoundWave == SoundWave)
 	{
-		Player->GetOutput().SetAudioSink(nullptr);
+		Player.SetAudioSink(nullptr);
 	}
 }
 
 
 void UMediaPlayer::HandleMediaTextureBeginDestroy(UMediaTexture& DestroyedMediaTexture)
 {
-	if (Player.IsValid() && (&DestroyedMediaTexture == VideoTexture))
+	if (&DestroyedMediaTexture == VideoTexture)
 	{
-		Player->GetOutput().SetVideoSink(nullptr);
+		Player.SetVideoSink(nullptr);
 	}
 }
 
@@ -802,20 +641,24 @@ void UMediaPlayer::HandlePlayerMediaEvent(EMediaEvent Event)
 		break;
 
 	case EMediaEvent::MediaOpened:
-		Player->GetControls().SetLooping(Loop);
-		Player->GetOutput().SetOverlaySink(this);
+		Player.SetLooping(Loop);
+
+		if (Overlays != nullptr)
+		{
+			Player.SetOverlaySink(Overlays);
+		}
 
 		if (SoundWave != nullptr)
 		{
-			Player->GetOutput().SetAudioSink(SoundWave);
+			Player.SetAudioSink(SoundWave);
 		}
 
 		if (VideoTexture != nullptr)
 		{
-			Player->GetOutput().SetVideoSink(VideoTexture);
+			Player.SetVideoSink(VideoTexture);
 		}
 
-		OnMediaOpened.Broadcast(CurrentUrl);
+		OnMediaOpened.Broadcast(Player.GetUrl());
 
 		if (PlayOnOpen)
 		{
@@ -824,7 +667,7 @@ void UMediaPlayer::HandlePlayerMediaEvent(EMediaEvent Event)
 		break;
 
 	case EMediaEvent::MediaOpenFailed:
-		OnMediaOpenFailed.Broadcast(CurrentUrl);
+		OnMediaOpenFailed.Broadcast(Player.GetUrl());
 
 		if (!Loop && (Playlist != nullptr))
 		{
@@ -848,15 +691,5 @@ void UMediaPlayer::HandlePlayerMediaEvent(EMediaEvent Event)
 	case EMediaEvent::PlaybackSuspended:
 		OnPlaybackSuspended.Broadcast();
 		break;
-
-	case EMediaEvent::TracksChanged:
-		SelectDefaultTracks();
-		break;
-	}
-
-	if ((Event == EMediaEvent::MediaOpened) ||
-		(Event == EMediaEvent::MediaOpenFailed))
-	{
-		UE_LOG(LogMediaAssets, Verbose, TEXT("Media Info:\n\n%s"), *Player->GetInfo());
 	}
 }
diff --git a/Engine/Source/Runtime/MediaAssets/Private/Assets/MediaSource.cpp b/Engine/Source/Runtime/MediaAssets/Private/Assets/MediaSource.cpp
--- a/Engine/Source/Runtime/MediaAssets/Private/Assets/MediaSource.cpp
+++ b/Engine/Source/Runtime/MediaAssets/Private/Assets/MediaSource.cpp
@@ -1,74 +1,6 @@
 // Copyright 1998-2017 Epic Games, Inc. All Rights Reserved.
 
 #include "MediaSource.h"
-#include "UObject/SequencerObjectVersion.h"
-#if WITH_EDITOR
-#include "Interfaces/ITargetPlatform.h"
-#endif
-
-
-/* UObject interface
- *****************************************************************************/
-
-void UMediaSource::GetAssetRegistryTags(TArray<FAssetRegistryTag>& OutTags) const
-{
-	FString Url = GetUrl();
-
-	if (!Url.IsEmpty())
-	{
-		OutTags.Add(FAssetRegistryTag("Url", Url, FAssetRegistryTag::TT_Alphabetical));
-	}
-}
-
-
-#if WITH_EDITOR
-void UMediaSource::GetAssetRegistryTagMetadata(TMap<FName, FAssetRegistryTagMetadata>& OutMetadata) const
-{
-}
-#endif
-
-
-void UMediaSource::Serialize(FArchive& Ar)
-{
-	Super::Serialize(Ar);
-
-	Ar.UsingCustomVersion(FSequencerObjectVersion::GUID);
-
-	if (Ar.IsLoading() && (Ar.CustomVer(FSequencerObjectVersion::GUID) < FSequencerObjectVersion::RenameMediaSourcePlatformPlayers))
-	{
-#if WITH_EDITORONLY_DATA
-		if (!Ar.IsFilterEditorOnly())
-		{
-			TMap<FString, FString> DummyPlatformPlayers;
-			Ar << DummyPlatformPlayers;
-		}
-#endif
-
-		FString DummyDefaultPlayer;
-		Ar << DummyDefaultPlayer;
-	}
-	else
-	{
-#if WITH_EDITORONLY_DATA
-		if (Ar.IsFilterEditorOnly())
-		{
-			if (Ar.IsSaving())
-			{
-				const FName* PlatformPlayerName = PlatformPlayerNames.Find(Ar.CookingTarget()->PlatformName());
-				PlayerName = (PlatformPlayerName != nullptr) ? *PlatformPlayerName : NAME_None;
-			}
-
-			Ar << PlayerName;
-		}
-		else
-		{
-			Ar << PlatformPlayerNames;
-		}
-#else
-		Ar << PlayerName;
-#endif
-	}
-}
 
 
 /* IMediaOptions interface
@@ -76,19 +8,7 @@ void UMediaSource::Serialize(FArchive& Ar)
 
 FName UMediaSource::GetDesiredPlayerName() const
 {
-#if WITH_EDITORONLY_DATA
-	const FString RunningPlatformName(FPlatformProperties::IniPlatformName());
-	const FName* PlatformPlayerName = PlatformPlayerNames.Find(RunningPlatformName);
-
-	if (PlatformPlayerName == nullptr)
-	{
-		return NAME_None;
-	}
-
-	return *PlatformPlayerName;
-#else
-	return PlayerName;
-#endif
+	return NAME_None;
 }
 
 
diff --git a/Engine/Source/Runtime/MediaAssets/Private/Assets/PlatformMediaSource.cpp b/Engine/Source/Runtime/MediaAssets/Private/Assets/PlatformMediaSource.cpp
--- a/Engine/Source/Runtime/MediaAssets/Private/Assets/PlatformMediaSource.cpp
+++ b/Engine/Source/Runtime/MediaAssets/Private/Assets/PlatformMediaSource.cpp
@@ -2,10 +2,12 @@
 
 #include "PlatformMediaSource.h"
 #include "UObject/SequencerObjectVersion.h"
+
 #if WITH_EDITOR
 	#include "Interfaces/ITargetPlatform.h"
 #endif
 
+
 /* UMediaSource interface
  *****************************************************************************/
 
@@ -82,40 +84,76 @@ UMediaSource* UPlatformMediaSource::GetMediaSource() const
 bool UPlatformMediaSource::GetMediaOption(const FName& Key, bool DefaultValue) const
 {
 	UMediaSource* PlatformMediaSource = GetMediaSource();
-	return (PlatformMediaSource != nullptr) ? PlatformMediaSource->GetMediaOption(Key, DefaultValue) : DefaultValue;
+	
+	if (PlatformMediaSource != nullptr)
+	{
+		return PlatformMediaSource->GetMediaOption(Key, DefaultValue);
+	}
+	
+	return Super::GetMediaOption(Key, DefaultValue);
 }
 
 
 double UPlatformMediaSource::GetMediaOption(const FName& Key, double DefaultValue) const
 {
 	UMediaSource* PlatformMediaSource = GetMediaSource();
-	return (PlatformMediaSource != nullptr) ? PlatformMediaSource->GetMediaOption(Key, DefaultValue) : DefaultValue;
+	
+	if (PlatformMediaSource != nullptr)
+	{
+		return PlatformMediaSource->GetMediaOption(Key, DefaultValue);
+	}
+
+	return Super::GetMediaOption(Key, DefaultValue);
 }
 
 
 int64 UPlatformMediaSource::GetMediaOption(const FName& Key, int64 DefaultValue) const
 {
 	UMediaSource* PlatformMediaSource = GetMediaSource();
-	return (PlatformMediaSource != nullptr) ? PlatformMediaSource->GetMediaOption(Key, DefaultValue) : DefaultValue;
+	
+	if (PlatformMediaSource != nullptr)
+	{
+		return PlatformMediaSource->GetMediaOption(Key, DefaultValue);
+	}
+
+	return Super::GetMediaOption(Key, DefaultValue);
 }
 
 
 FString UPlatformMediaSource::GetMediaOption(const FName& Key, const FString& DefaultValue) const
 {
 	UMediaSource* PlatformMediaSource = GetMediaSource();
-	return (PlatformMediaSource != nullptr) ? PlatformMediaSource->GetMediaOption(Key, DefaultValue) : DefaultValue;
+	
+	if (PlatformMediaSource != nullptr)
+	{
+		return PlatformMediaSource->GetMediaOption(Key, DefaultValue);
+	}
+
+	return Super::GetMediaOption(Key, DefaultValue);
 }
 
 
 FText UPlatformMediaSource::GetMediaOption(const FName& Key, const FText& DefaultValue) const
 {
 	UMediaSource* PlatformMediaSource = GetMediaSource();
-	return (PlatformMediaSource != nullptr) ? PlatformMediaSource->GetMediaOption(Key, DefaultValue) : DefaultValue;
+	
+	if (PlatformMediaSource != nullptr)
+	{
+		return PlatformMediaSource->GetMediaOption(Key, DefaultValue);
+	}
+
+	return Super::GetMediaOption(Key, DefaultValue);
 }
 
 
 bool UPlatformMediaSource::HasMediaOption(const FName& Key) const
 {
 	UMediaSource* PlatformMediaSource = GetMediaSource();
-	return (PlatformMediaSource != nullptr) ? PlatformMediaSource->HasMediaOption(Key) : false;
+	
+	if (PlatformMediaSource != nullptr)
+	{
+		return PlatformMediaSource->HasMediaOption(Key);
+	}
+
+	return Super::HasMediaOption(Key);
 }
diff --git a/Engine/Source/Runtime/MediaAssets/Private/Assets/StreamMediaSource.cpp b/Engine/Source/Runtime/MediaAssets/Private/Assets/StreamMediaSource.cpp
--- a/Engine/Source/Runtime/MediaAssets/Private/Assets/StreamMediaSource.cpp
+++ b/Engine/Source/Runtime/MediaAssets/Private/Assets/StreamMediaSource.cpp
@@ -3,6 +3,9 @@
 #include "StreamMediaSource.h"
 
 
+/* UMediaSource overrides
+ *****************************************************************************/
+
 FString UStreamMediaSource::GetUrl() const
 {
 	return StreamUrl;
diff --git a/Engine/Source/Runtime/MediaAssets/Private/Misc/MediaTextureResource.cpp b/Engine/Source/Runtime/MediaAssets/Private/Misc/MediaTextureResource.cpp
--- a/Engine/Source/Runtime/MediaAssets/Private/Misc/MediaTextureResource.cpp
+++ b/Engine/Source/Runtime/MediaAssets/Private/Misc/MediaTextureResource.cpp
@@ -53,7 +53,7 @@ void* FMediaTextureResource::AcquireBuffer()
 		return TripleBuffer.GetWriteBuffer().LockedData;
 	}
 
-	check(IsInActualRenderingThread());
+	check(IsInRenderingThread());
 
 	uint32 OutStride = 0;
 	TRefCountPtr<FRHITexture2D>& Texture = RequiresConversion
@@ -71,7 +71,7 @@ void FMediaTextureResource::DisplayBuffer()
 		return;
 	}
 
-	if (IsInActualRenderingThread())
+	if (IsInRenderingThread())
 	{
 		SwapResource();
 	}
@@ -102,7 +102,7 @@ void FMediaTextureResource::InitializeBuffer(FIntPoint OutputDim, FIntPoint Buff
 {
 	State = EState::Initializing;
 
-	if (IsInActualRenderingThread())
+	if (IsInRenderingThread())
 	{
 		InitializeResource(OutputDim, BufferDim, Format, Mode);
 	}
@@ -128,7 +128,7 @@ void FMediaTextureResource::ReleaseBuffer()
 	}
 	else
 	{
-		check(IsInActualRenderingThread());
+		check(IsInRenderingThread());
 
 		TRefCountPtr<FRHITexture2D>& Texture = RequiresConversion
 			? BufferResources[2].RenderTarget
@@ -148,7 +148,7 @@ void FMediaTextureResource::ShutdownBuffer()
 
 	State = EState::ShuttingDown;
 
-	if (IsInActualRenderingThread())
+	if (IsInRenderingThread())
 	{
 		ReleaseDynamicRHI();
 	}
@@ -185,7 +185,7 @@ void FMediaTextureResource::UpdateBuffer(const uint8* Data, uint32 /*Pitch*/)
 			TripleBuffer.SwapWriteBuffers();
 		}
 	}
-	else if (IsInActualRenderingThread())
+	else if (IsInRenderingThread())
 	{
 		if (RequiresConversion)
 		{
@@ -221,7 +221,7 @@ void FMediaTextureResource::UpdateTextures(FRHITexture* RenderTarget, FRHITextur
 		Resource.ShaderResource = (ShaderResource != nullptr) ? ShaderResource->GetTexture2D() : nullptr;
 	}
 
-	if (IsInActualRenderingThread())
+	if (IsInRenderingThread())
 	{
 		SetResource(Resource);
 	}
@@ -283,63 +283,77 @@ void FMediaTextureResource::InitDynamicRHI()
 
 	// determine buffer resource pixel format
 	EPixelFormat BufferFormat;
-	bool SrgbBuffer = false;
+	bool NeedsResize;
+	bool NeedsTransform;
+	bool SrgbBuffer;
 
 	switch (SinkFormat)
 	{
-	case EMediaTextureSinkFormat::CharNV12:
-	case EMediaTextureSinkFormat::CharNV21:
-		BufferFormat = PF_G8;
-		BufferPitch = BufferDimensions.X;
-		RequiresConversion = true;
-		break;
-
-	case EMediaTextureSinkFormat::CharUYVY:
-	case EMediaTextureSinkFormat::CharYUY2:
-	case EMediaTextureSinkFormat::CharYVYU:
-		BufferFormat = PF_B8G8R8A8;
-		BufferPitch = BufferDimensions.X * 4;
-		RequiresConversion = true;
-		break;
-
 	case EMediaTextureSinkFormat::CharAYUV:
 		BufferFormat = PF_B8G8R8A8;
 		BufferPitch = BufferDimensions.X * 4;
-		RequiresConversion = true;
+		NeedsResize = false;
+		NeedsTransform = true;
+		SrgbBuffer = false;
 		break;
 
 	case EMediaTextureSinkFormat::CharBGRA:
 		BufferFormat = PF_B8G8R8A8;
 		BufferPitch = BufferDimensions.X * 4;
-		RequiresConversion = (BufferDimensions != OutputDimensions);
+		NeedsResize = (BufferDimensions != OutputDimensions);
+		NeedsTransform = false;
 		SrgbBuffer = true;
 		break;
 
 	case EMediaTextureSinkFormat::CharBMP:
 		BufferFormat = PF_B8G8R8A8;
 		BufferPitch = BufferDimensions.X * 4;
-		RequiresConversion = true;
+		NeedsResize = true;
+		NeedsTransform = false;
 		SrgbBuffer = true;
 		break;
 
+	case EMediaTextureSinkFormat::CharNV12:
+	case EMediaTextureSinkFormat::CharNV21:
+		BufferFormat = PF_G8;
+		BufferPitch = BufferDimensions.X;
+		NeedsResize = false;
+		NeedsTransform = true;
+		SrgbBuffer = false;
+		break;
+
+	case EMediaTextureSinkFormat::CharUYVY:
+	case EMediaTextureSinkFormat::CharYUY2:
+	case EMediaTextureSinkFormat::CharYVYU:
+		BufferFormat = PF_B8G8R8A8;
+		BufferPitch = BufferDimensions.X * 4;
+		NeedsResize = false;
+		NeedsTransform = true;
+		SrgbBuffer = false;
+		break;
+
 	case EMediaTextureSinkFormat::FloatRGB:
 		BufferFormat = PF_FloatRGB;
 		BufferPitch = BufferDimensions.X * 12;
-		RequiresConversion = (BufferDimensions != OutputDimensions);
+		NeedsResize = (BufferDimensions != OutputDimensions);
+		NeedsTransform = false;
 		SrgbBuffer = true;
 		break;
 
 	case EMediaTextureSinkFormat::FloatRGBA:
 		BufferFormat = PF_FloatRGBA;
 		BufferPitch = BufferDimensions.X * 16;
-		RequiresConversion = (BufferDimensions != OutputDimensions);
+		NeedsResize = (BufferDimensions != OutputDimensions);
+		NeedsTransform = false;
 		SrgbBuffer = true;
 		break;
 
 	default:
-		return;
+		return; // unsupported sink format
 	}
 
+	RequiresConversion = NeedsResize || NeedsTransform;
+
 	// determine how many buffer resources we need
 	int32 StartIndex;
 	
@@ -441,7 +455,7 @@ void FMediaTextureResource::InitDynamicRHI()
 		SetRenderTarget(CommandList, RenderTargetTextureRHI, FTextureRHIRef());
 		CommandList.SetViewport(0, 0, 0.0f, OutputDimensions.X, OutputDimensions.Y, 1.0f);
 		CommandList.ClearColorTexture(RenderTargetTextureRHI, BufferClearColor, FIntRect());
-		CommandList.CopyToResolveTarget(TextureRHI, TextureRHI, true, FResolveParams());
+		CommandList.TransitionResource(EResourceTransitionAccess::EReadable, TextureRHI);
 	}
 
 	CacheResourceSize();
@@ -549,9 +563,9 @@ void FMediaTextureResource::ConvertResource(const FResource& Resource)
 {
 	check(IsInRenderingThread());
 
-	FRHICommandListImmediate& RHICmdList = FRHICommandListExecutor::GetImmediateCommandList();
+	FRHICommandListImmediate& CommandList = FRHICommandListExecutor::GetImmediateCommandList();
 
-	SCOPED_DRAW_EVENT(RHICmdList, MediaTextureConvertResource);
+	SCOPED_DRAW_EVENT(CommandList, MediaTextureConvertResource);
 
 	// configure media shaders
 	auto ShaderMap = GetGlobalShaderMap(GMaxRHIFeatureLevel);
@@ -563,9 +577,14 @@ void FMediaTextureResource::ConvertResource(const FResource& Resource)
 		{
 			TShaderMapRef<FAYUVConvertPS> ConvertShader(ShaderMap);
 			static FGlobalBoundShaderState BoundShaderState;
+			SetGlobalBoundShaderState(CommandList, GMaxRHIFeatureLevel, BoundShaderState, GMediaVertexDeclaration.VertexDeclarationRHI, *VertexShader, *ConvertShader);
 
-			SetGlobalBoundShaderState(RHICmdList, GMaxRHIFeatureLevel, BoundShaderState, GMediaVertexDeclaration.VertexDeclarationRHI, *VertexShader, *ConvertShader);
-			ConvertShader->SetParameters(RHICmdList, Resource.ShaderResource);
+			ConvertShader->SetParameters(
+				CommandList,
+				Resource.ShaderResource,
+				MediaShaders::YuvToSrgbDefault,
+				Owner.SRGB
+			);
 		}
 		break;
 
@@ -573,9 +592,15 @@ void FMediaTextureResource::ConvertResource(const FResource& Resource)
 		{
 			TShaderMapRef<FNV12ConvertPS> ConvertShader(ShaderMap);
 			static FGlobalBoundShaderState BoundShaderState;
-
-			SetGlobalBoundShaderState(RHICmdList, GMaxRHIFeatureLevel, BoundShaderState, GMediaVertexDeclaration.VertexDeclarationRHI, *VertexShader, *ConvertShader);
-			ConvertShader->SetParameters(RHICmdList, Resource.ShaderResource, OutputDimensions);
+			SetGlobalBoundShaderState(CommandList, GMaxRHIFeatureLevel, BoundShaderState, GMediaVertexDeclaration.VertexDeclarationRHI, *VertexShader, *ConvertShader);
+
+			ConvertShader->SetParameters(
+				CommandList,
+				Resource.ShaderResource,
+				OutputDimensions,
+				MediaShaders::YuvToSrgbDefault,
+				Owner.SRGB
+			);
 		}
 		break;
 
@@ -583,9 +608,15 @@ void FMediaTextureResource::ConvertResource(const FResource& Resource)
 		{
 			TShaderMapRef<FNV21ConvertPS> ConvertShader(ShaderMap);
 			static FGlobalBoundShaderState BoundShaderState;
-
-			SetGlobalBoundShaderState(RHICmdList, GMaxRHIFeatureLevel, BoundShaderState, GMediaVertexDeclaration.VertexDeclarationRHI, *VertexShader, *ConvertShader);
-			ConvertShader->SetParameters(RHICmdList, Resource.ShaderResource, OutputDimensions);
+			SetGlobalBoundShaderState(CommandList, GMaxRHIFeatureLevel, BoundShaderState, GMediaVertexDeclaration.VertexDeclarationRHI, *VertexShader, *ConvertShader);
+
+			ConvertShader->SetParameters(
+				CommandList,
+				Resource.ShaderResource,
+				OutputDimensions,
+				MediaShaders::YuvToSrgbDefault,
+				Owner.SRGB
+			);
 		}
 		break;
 
@@ -593,9 +624,14 @@ void FMediaTextureResource::ConvertResource(const FResource& Resource)
 		{
 			TShaderMapRef<FUYVYConvertPS> ConvertShader(ShaderMap);
 			static FGlobalBoundShaderState BoundShaderState;
+			SetGlobalBoundShaderState(CommandList, GMaxRHIFeatureLevel, BoundShaderState, GMediaVertexDeclaration.VertexDeclarationRHI, *VertexShader, *ConvertShader);
 
-			SetGlobalBoundShaderState(RHICmdList, GMaxRHIFeatureLevel, BoundShaderState, GMediaVertexDeclaration.VertexDeclarationRHI, *VertexShader, *ConvertShader);
-			ConvertShader->SetParameters(RHICmdList, Resource.ShaderResource);
+			ConvertShader->SetParameters(
+				CommandList,
+				Resource.ShaderResource,
+				MediaShaders::YuvToSrgbDefault,
+				Owner.SRGB
+			);
 		}
 		break;
 
@@ -603,9 +639,15 @@ void FMediaTextureResource::ConvertResource(const FResource& Resource)
 		{
 			TShaderMapRef<FYUY2ConvertPS> ConvertShader(ShaderMap);
 			static FGlobalBoundShaderState BoundShaderState;
-
-			SetGlobalBoundShaderState(RHICmdList, GMaxRHIFeatureLevel, BoundShaderState, GMediaVertexDeclaration.VertexDeclarationRHI, *VertexShader, *ConvertShader);
-			ConvertShader->SetParameters(RHICmdList, Resource.ShaderResource, OutputDimensions);
+			SetGlobalBoundShaderState(CommandList, GMaxRHIFeatureLevel, BoundShaderState, GMediaVertexDeclaration.VertexDeclarationRHI, *VertexShader, *ConvertShader);
+
+			ConvertShader->SetParameters(
+				CommandList,
+				Resource.ShaderResource,
+				OutputDimensions,
+				MediaShaders::YuvToSrgbDefault,
+				Owner.SRGB
+			);
 		}
 		break;
 
@@ -613,9 +655,14 @@ void FMediaTextureResource::ConvertResource(const FResource& Resource)
 		{
 			TShaderMapRef<FYVYUConvertPS> ConvertShader(ShaderMap);
 			static FGlobalBoundShaderState BoundShaderState;
+			SetGlobalBoundShaderState(CommandList, GMaxRHIFeatureLevel, BoundShaderState, GMediaVertexDeclaration.VertexDeclarationRHI, *VertexShader, *ConvertShader);
 
-			SetGlobalBoundShaderState(RHICmdList, GMaxRHIFeatureLevel, BoundShaderState, GMediaVertexDeclaration.VertexDeclarationRHI, *VertexShader, *ConvertShader);
-			ConvertShader->SetParameters(RHICmdList, Resource.ShaderResource);
+			ConvertShader->SetParameters(
+				CommandList,
+				Resource.ShaderResource,
+				MediaShaders::YuvToSrgbDefault,
+				Owner.SRGB
+			);
 		}
 		break;
 
@@ -626,8 +673,8 @@ void FMediaTextureResource::ConvertResource(const FResource& Resource)
 			TShaderMapRef<FRGBConvertPS> ConvertShader(ShaderMap);
 			static FGlobalBoundShaderState BoundShaderState;
 
-			SetGlobalBoundShaderState(RHICmdList, GMaxRHIFeatureLevel, BoundShaderState, GMediaVertexDeclaration.VertexDeclarationRHI, *VertexShader, *ConvertShader);
-			ConvertShader->SetParameters(RHICmdList, Resource.ShaderResource, OutputDimensions);
+			SetGlobalBoundShaderState(CommandList, GMaxRHIFeatureLevel, BoundShaderState, GMediaVertexDeclaration.VertexDeclarationRHI, *VertexShader, *ConvertShader);
+			ConvertShader->SetParameters(CommandList, Resource.ShaderResource, OutputDimensions, Owner.SRGB);
 		}
 		break;
 
@@ -636,8 +683,8 @@ void FMediaTextureResource::ConvertResource(const FResource& Resource)
 			TShaderMapRef<FBMPConvertPS> ConvertShader(ShaderMap);
 			static FGlobalBoundShaderState BoundShaderState;
 
-			SetGlobalBoundShaderState(RHICmdList, GMaxRHIFeatureLevel, BoundShaderState, GMediaVertexDeclaration.VertexDeclarationRHI, *VertexShader, *ConvertShader);
-			ConvertShader->SetParameters(RHICmdList, Resource.ShaderResource, OutputDimensions);
+			SetGlobalBoundShaderState(CommandList, GMaxRHIFeatureLevel, BoundShaderState, GMediaVertexDeclaration.VertexDeclarationRHI, *VertexShader, *ConvertShader);
+			ConvertShader->SetParameters(CommandList, Resource.ShaderResource, OutputDimensions);
 		}
 		break;
 
@@ -658,16 +705,16 @@ void FMediaTextureResource::ConvertResource(const FResource& Resource)
 		Vertices[3].TextureCoordinate.Set(1.0f, 1.0f);
 	}
 
-	RHICmdList.SetDepthStencilState(TStaticDepthStencilState<false, CF_Always>::GetRHI());
-	RHICmdList.SetRasterizerState(TStaticRasterizerState<>::GetRHI());
-	RHICmdList.SetBlendState(TStaticBlendStateWriteMask<CW_RGBA, CW_NONE, CW_NONE, CW_NONE, CW_NONE, CW_NONE, CW_NONE, CW_NONE>::GetRHI());
-	RHICmdList.SetViewport(0, 0, 0.0f, OutputDimensions.X, OutputDimensions.Y, 1.0f);
+	CommandList.SetDepthStencilState(TStaticDepthStencilState<false, CF_Always>::GetRHI());
+	CommandList.SetRasterizerState(TStaticRasterizerState<>::GetRHI());
+	CommandList.SetBlendState(TStaticBlendStateWriteMask<CW_RGBA, CW_NONE, CW_NONE, CW_NONE, CW_NONE, CW_NONE, CW_NONE, CW_NONE>::GetRHI());
+	CommandList.SetViewport(0, 0, 0.0f, OutputDimensions.X, OutputDimensions.Y, 1.0f);
 
 	FRHITexture* RenderTarget = RenderTargetTextureRHI.GetReference();
-	SetRenderTargets(RHICmdList, 1, &RenderTarget, nullptr, ESimpleRenderTargetMode::EExistingColorAndDepth, FExclusiveDepthStencil::DepthNop_StencilNop);
-	DrawPrimitiveUP(RHICmdList, PT_TriangleStrip, 2, Vertices, sizeof(Vertices[0]));
+	SetRenderTargets(CommandList, 1, &RenderTarget, nullptr, ESimpleRenderTargetMode::EExistingColorAndDepth, FExclusiveDepthStencil::DepthNop_StencilNop);
+	DrawPrimitiveUP(CommandList, PT_TriangleStrip, 2, Vertices, sizeof(Vertices[0]));
 
-	RHICmdList.CopyToResolveTarget(RenderTargetTextureRHI, RenderTargetTextureRHI, true, FResolveParams());
+	CommandList.TransitionResource(EResourceTransitionAccess::EReadable, RenderTargetTextureRHI);
 }
 
 
@@ -686,7 +733,10 @@ void FMediaTextureResource::InitializeResource(FIntPoint OutputDim, FIntPoint Bu
 {
 	check(IsInRenderingThread());
 
-	if ((OutputDim == OutputDimensions) && (BufferDim == BufferDimensions) && (Format == SinkFormat) && (Mode == SinkMode))
+	if ((OutputDim == OutputDimensions) &&
+		(BufferDim == BufferDimensions) &&
+		(Format == SinkFormat) &&
+		(Mode == SinkMode))
 	{
 		// reuse existing resources
 		State = (OutputDimensions.GetMin() > 0) ? EState::Initialized : EState::ShutDown;
diff --git a/Engine/Source/Runtime/MediaAssets/Public/BaseMediaSource.h b/Engine/Source/Runtime/MediaAssets/Public/BaseMediaSource.h
new file mode 100644
--- /dev/null
+++ b/Engine/Source/Runtime/MediaAssets/Public/BaseMediaSource.h
@@ -0,0 +1,54 @@
+// Copyright 1998-2017 Epic Games, Inc. All Rights Reserved.
+
+#pragma once
+
+#include "CoreMinimal.h"
+#include "UObject/ObjectMacros.h"
+#include "UObject/Object.h"
+#include "MediaSource.h"
+
+#include "BaseMediaSource.generated.h"
+
+
+/**
+ * Base class for concrete media sources.
+ */
+UCLASS(Abstract, BlueprintType, hidecategories=(Object))
+class MEDIAASSETS_API UBaseMediaSource
+	: public UMediaSource
+{
+	GENERATED_BODY()
+
+public:
+
+#if WITH_EDITORONLY_DATA
+
+	/** Override native media player plug-ins per platform (Empty = find one automatically). */
+	UPROPERTY(transient, BlueprintReadWrite, EditAnywhere, Category=Platforms, Meta=(DisplayName="Player Overrides"))
+	TMap<FString, FName> PlatformPlayerNames;
+
+#endif
+
+public:
+
+	//~ UObject interface
+
+	virtual void GetAssetRegistryTags(TArray<FAssetRegistryTag>& OutTags) const override;
+	virtual void Serialize(FArchive& Ar) override;
+
+#if WITH_EDITOR
+	virtual void GetAssetRegistryTagMetadata(TMap<FName, FAssetRegistryTagMetadata>& OutMetadata) const override;
+#endif
+
+public:
+
+	//~ IMediaOptions interface
+
+	virtual FName GetDesiredPlayerName() const override;
+
+private:
+
+	/** Name of the desired native media player (Empty = find one automatically). */
+	UPROPERTY(transient)
+	FName PlayerName;
+};
diff --git a/Engine/Source/Runtime/MediaAssets/Public/FileMediaSource.h b/Engine/Source/Runtime/MediaAssets/Public/FileMediaSource.h
--- a/Engine/Source/Runtime/MediaAssets/Public/FileMediaSource.h
+++ b/Engine/Source/Runtime/MediaAssets/Public/FileMediaSource.h
@@ -5,12 +5,14 @@
 #include "CoreMinimal.h"
 #include "UObject/ObjectMacros.h"
 #include "UObject/ScriptMacros.h"
-#include "MediaSource.h"
+#include "BaseMediaSource.h"
+
 #include "FileMediaSource.generated.h"
 
+
 UCLASS(BlueprintType)
 class MEDIAASSETS_API UFileMediaSource
-	: public UMediaSource
+	: public UBaseMediaSource
 {
 	GENERATED_BODY()
 
@@ -58,6 +60,7 @@ public:
 	//~ IMediaOptions interface
 
 	virtual bool GetMediaOption(const FName& Key, bool DefaultValue) const override;
+	virtual bool HasMediaOption(const FName& Key) const override;
 
 public:
 
@@ -66,6 +69,11 @@ public:
 	virtual FString GetUrl() const override;
 	virtual bool Validate() const override;
 
+public:
+
+	/** Name of the PrecacheFile media option. */
+	static FName PrecacheFileOption;
+
 protected:
 
 	/** Get the full path to the file. */
diff --git a/Engine/Source/Runtime/MediaAssets/Public/MediaOverlays.h b/Engine/Source/Runtime/MediaAssets/Public/MediaOverlays.h
new file mode 100644
--- /dev/null
+++ b/Engine/Source/Runtime/MediaAssets/Public/MediaOverlays.h
@@ -0,0 +1,140 @@
+// Copyright 1998-2016 Epic Games, Inc. All Rights Reserved.
+
+#pragma once
+
+#include "CoreMinimal.h"
+#include "UObject/Object.h"
+#include "UObject/ObjectMacros.h"
+#include "UObject/ScriptMacros.h"
+#include "IMediaOverlaySink.h"
+
+#include "MediaOverlays.generated.h"
+
+
+USTRUCT(BlueprintType)
+struct FMediaPlayerOverlay
+{
+	GENERATED_BODY()
+
+	/** Whether the text position is set. */
+	UPROPERTY()
+	bool HasPosition;
+
+	/** The text position. */
+	UPROPERTY()
+	FVector2D Position;
+
+	/** The overlay text. */
+	UPROPERTY()
+	FText Text;
+};
+
+
+/**
+* Implements an asset that collects overlay texts from UMediaPlayer assets.
+*/
+UCLASS(BlueprintType, hidecategories=(Object))
+class MEDIAASSETS_API UMediaOverlays
+	: public UObject
+	, public IMediaOverlaySink
+{
+	/** Text overlay entry. */
+	struct FOverlay
+	{
+		FTimespan Duration;
+		TOptional<FVector2D> Position;
+		FText Text;
+		FTimespan Time;
+		EMediaOverlayType Type;
+	};
+
+	GENERATED_BODY()
+
+public:
+
+	/**
+	 * Get the current caption text overlays for the specified time, if any.
+	 *
+	 * @param OutCaptions Will contain the caption text overlays.
+	 * @param Time The for which to get the overlays.
+	 * @see GetSubtitles, GetTexts
+	 */
+	UFUNCTION(BlueprintCallable, Category="Media|MediaPlayer")
+	void GetCaptions(TArray<FMediaPlayerOverlay>& OutCaptions, FTimespan Time) const
+	{
+		GetItems(EMediaOverlayType::Caption, OutCaptions, Time);
+	}
+
+	/**
+	 * Get the current subtitle text overlays for the specified time, if any.
+	 *
+	 * @param OutSubtitles Will contain the caption text overlays.
+	 * @param Time The for which to get the overlays.
+	 * @see GetCaptions, GetTexts
+	 */
+	UFUNCTION(BlueprintCallable, Category="Media|MediaPlayer")
+	void GetSubtitles(TArray<FMediaPlayerOverlay>& OutSubtitles, FTimespan Time) const
+	{
+		GetItems(EMediaOverlayType::Subtitle, OutSubtitles, Time);
+	}
+
+	/**
+	 * Get the current generic text overlays for the specified time, if any.
+	 *
+	 * @param OutTexts Will contain the text overlays.
+	 * @param Time The for which to get the overlays.
+	 * @see GetCaptions, GetSubtitles
+	 */
+	UFUNCTION(BlueprintCallable, Category="Media|MediaPlayer")
+	void GetTexts(TArray<FMediaPlayerOverlay>& OutTexts, FTimespan Time) const
+	{
+		GetItems(EMediaOverlayType::Text, OutTexts, Time);
+	}
+
+public:
+
+	/**
+	* Get the text overlay items of the given type at the specified time.
+	*
+	* @param Type The type of overlays to get, i.e. captions, subtitles, etc.
+	* @param Time The for which to get the overlays.
+	* @param OutItems Will contain the text overlay items.
+	*/
+	void GetItems(EMediaOverlayType Type, TArray<FMediaPlayerOverlay>& OutItems, FTimespan Time) const;
+
+public:
+
+	/**
+	 * Get the event delegate that is invoked when this asset is being destroyed.
+	 *
+	 * @return The delegate.
+	 */
+	DECLARE_EVENT_OneParam(UMediaOverlays, FOnBeginDestroy, UMediaOverlays& /*DestroyedOverlays*/)
+	FOnBeginDestroy& OnBeginDestroy()
+	{
+		return BeginDestroyEvent;
+	}
+
+public:
+
+	//~ UObject interface.
+
+	virtual void BeginDestroy() override;
+
+public:
+
+	//~ IMediaOverlaySink interface
+
+	virtual bool InitializeOverlaySink() override;
+	virtual void AddOverlaySinkText(const FText& Text, EMediaOverlayType Type, FTimespan Time, FTimespan Duration, TOptional<FVector2D> Position) override;
+	virtual void ClearOverlaySinkText() override;
+	virtual void ShutdownOverlaySink() override;
+
+private:
+
+	/** An event delegate that is invoked when this asset is being destroyed. */
+	FOnBeginDestroy BeginDestroyEvent;
+
+	/** Queued text overlays. */
+	TArray<FOverlay> Overlays;
+};
diff --git a/Engine/Source/Runtime/MediaAssets/Public/MediaPlayer.h b/Engine/Source/Runtime/MediaAssets/Public/MediaPlayer.h
--- a/Engine/Source/Runtime/MediaAssets/Public/MediaPlayer.h
+++ b/Engine/Source/Runtime/MediaAssets/Public/MediaPlayer.h
@@ -3,14 +3,18 @@
 #pragma once
 
 #include "CoreMinimal.h"
+#include "Containers/Ticker.h"
 #include "UObject/ObjectMacros.h"
 #include "UObject/Object.h"
 #include "UObject/ScriptMacros.h"
 #include "IMediaOverlaySink.h"
+#include "MediaPlayerBase.h"
 #include "MediaPlayer.generated.h"
 
-class IMediaOptions;
+
+class FMediaPlayerBase;
 class IMediaPlayer;
+class UMediaOverlays;
 class UMediaPlaylist;
 class UMediaSoundWave;
 class UMediaSource;
@@ -58,32 +62,13 @@ enum class EMediaPlayerTrack : uint8
 };
 
 
-USTRUCT(BlueprintType)
-struct FMediaPlayerOverlay
-{
-	GENERATED_BODY()
-
-	/** Whether the text position is set. */
-	UPROPERTY()
-	bool HasPosition;
-
-	/** The text position. */
-	UPROPERTY()
-	FVector2D Position;
-
-	/** The overlay text. */
-	UPROPERTY()
-	FText Text;
-};
-
-
 /**
  * Implements a media player asset that can play movies and other media sources.
  */
 UCLASS(BlueprintType, hidecategories=(Object))
 class MEDIAASSETS_API UMediaPlayer
 	: public UObject
-	, public IMediaOverlaySink
+	, public FTickerObjectBase
 {
 	GENERATED_UCLASS_BODY()
 
@@ -103,23 +88,24 @@ public:
 	/**
 	 * Check whether the specified media source can be played by this player.
 	 *
-	 * If a desired player name is set for this player (DesiredPlayerName), it will
-	 * only check whether that particular player type can play the specified source.
+	 * If a desired player name is set for this player, it will only check
+	 * whether that particular player type can play the specified source.
 	 *
 	 * @param MediaSource The media source to check.
 	 * @return true if the media source can be opened, false otherwise.
-	 * @see CanPlayUrl, DesiredPlayerName
+	 * @see CanPlayUrl, SetDesiredPlayerName
 	 */
+	UFUNCTION(BlueprintCallable, Category = "Media|MediaPlayer")
 	bool CanPlaySource(UMediaSource* MediaSource);
 
 	/**
 	 * Check whether the specified URL can be played by this player.
 	 *
-	 * If a desired player name is set for this player (DesiredPlayerName), it will
-	 * only check whether that particular player type can play the specified URL.
+	 * If a desired player name is set for this player, it will only check
+	 * whether that particular player type can play the specified URL.
 	 *
 	 * @param Url The URL to check.
-	 * @see CanPlaySource, DesiredPlayerName
+	 * @see CanPlaySource, SetDesiredPlayerName
 	 */
 	UFUNCTION(BlueprintCallable, Category="Media|MediaPlayer")
 	bool CanPlayUrl(const FString& Url);
@@ -133,16 +119,13 @@ public:
 	void Close();
 
 	/**
-	 * Get the current caption text overlays, if any.
+	 * Get the name of the current desired native player.
 	 *
-	 * @param OutCaptions Will contain the caption text overlays.
-	 * @see GetSubtitles, GetTexts
+	 * @return The name of the desired player, or NAME_None if not set.
+	 * @see SetDesiredPlayerName
 	 */
-	UFUNCTION(BlueprintCallable, Category="Media|MediaPlayer")
-	void GetCaptions(TArray<FMediaPlayerOverlay>& OutCaptions) const
-	{
-		GetOverlays(EMediaOverlayType::Caption, OutCaptions);
-	}
+	UFUNCTION(BlueprintCallable, Category = "Media|MediaPlayer")
+	FName GetDesiredPlayerName() const;
 
 	/**
 	 * Get the media's duration.
@@ -210,30 +193,6 @@ public:
 	UFUNCTION(BlueprintCallable, Category="Media|MediaPlayer")
 	int32 GetSelectedTrack(EMediaPlayerTrack TrackType) const;
 
-	/**
-	 * Get the current subtitle text overlays, if any.
-	 *
-	 * @param OutSubtitles Will contain the caption text overlays.
-	 * @see GetCaptions, GetTexts
-	 */
-	UFUNCTION(BlueprintCallable, Category="Media|MediaPlayer")
-	void GetSubtitles(TArray<FMediaPlayerOverlay>& OutSubtitles) const
-	{
-		GetOverlays(EMediaOverlayType::Subtitle, OutSubtitles);
-	}
-
-	/**
-	 * Get the current generic text overlays, if any.
-	 *
-	 * @param OutTexts Will contain the text overlays.
-	 * @see GetCaptions, GetSubtitles
-	 */
-	UFUNCTION(BlueprintCallable, Category="Media|MediaPlayer")
-	void GetTexts(TArray<FMediaPlayerOverlay>& OutTexts) const
-	{
-		GetOverlays(EMediaOverlayType::Text, OutTexts);
-	}
-
 	/**
 	 * Get the media's current playback time.
 	 *
@@ -272,10 +231,7 @@ public:
 	 * @see OpenUrl
 	 */
 	UFUNCTION(BlueprintCallable, Category="Media|MediaPlayer")
-	const FString& GetUrl() const
-	{
-		return CurrentUrl;
-	}
+	const FString& GetUrl() const;
 
 	/**
 	 * Checks whether playback is looping.
@@ -483,6 +439,15 @@ public:
 	UFUNCTION(BlueprintCallable, Category="Media|MediaPlayer")
 	bool SelectTrack(EMediaPlayerTrack TrackType, int32 TrackIndex);
 
+	/**
+	 * Set the name of the desired native player.
+	 *
+	 * @param PlayerName The name of the player to set.
+	 * @see GetDesiredPlayerName
+	 */
+	UFUNCTION(BlueprintCallable, Category = "Media|MediaPlayer")
+	void SetDesiredPlayerName(FName PlayerName);
+
 	/**
 	 * Enables or disables playback looping.
 	 *
@@ -493,6 +458,15 @@ public:
 	UFUNCTION(BlueprintCallable, Category="Media|MediaPlayer")
 	bool SetLooping(bool Looping);
 
+	/**
+	 * Assign the given overlays asset to the player's overlay sink.
+	 *
+	 * @param NewOverlays The overlays asset to set.
+	 * @see SetVideoTexture
+	 */
+	UFUNCTION(BlueprintCallable, Category="Media|MediaPlayer")
+	void SetOverlays(UMediaOverlays* NewOverlays);
+
 	/**
 	 * Changes the media's playback rate.
 	 *
@@ -589,39 +563,39 @@ public:
 
 public:
 
-	/** Get an event delegate that is invoked when a media event occurred. */
-	DECLARE_EVENT_OneParam(UMediaPlayer, FOnMediaEvent, EMediaEvent /*Event*/)
-	FOnMediaEvent& OnMediaEvent()
-	{
-		return MediaEvent;
-	}
-
-public:
+	/**
+	 * Get the player implementation of this object.
+	 *
+	 * @return The player implementation.
+	 */
+	FMediaPlayerBase& GetBasePlayer();
 
 	/**
-	 * Get the text overlays of the specified type.
+	 * Get the URL of the media that was last attempted to be opened.
 	 *
-	 * @param Type The type of overlays to get, i.e. captions, subtitles, etc.
-	 * @param OutOverlays Will contain the text overlays.
+	 * @return The last attempted URL.
 	 */
-	void GetOverlays(EMediaOverlayType Type, TArray<FMediaPlayerOverlay>& OutOverlays) const;
+	const FString& GetLastUrl()
+	{
+		return LastUrl;
+	}
 
 	/**
-	 * Get the low-level player associated with this object.
+	 * Get the overlays assets that receives the captions and subtitles.
 	 *
-	 * @return The player, or nullptr if no player was created.
-	 * @see GetPlaylist, GetPlaylistIndex
+	 * @return Overlays asset.
+	 * @see GetPlayer, GetPlaylist, GetVideoTexture
 	 */
-	TSharedPtr<IMediaPlayer> GetPlayer() const
+	UMediaOverlays* GetOverlays() const
 	{
-		return Player;
+		return Overlays;
 	}
 
 	/**
 	 * Get the currently active play list.
 	 *
 	 * @return The play list, if any.
-	 * @see GetPlayer, GetPlaylistIndex
+	 * @see GetOverlays, GetPlayer, GetPlaylistIndex
 	 */
 	UMediaPlaylist* GetPlaylist() const
 	{
@@ -632,7 +606,7 @@ public:
 	 * Get the current play list index.
 	 *
 	 * @return Play list index.
-	 * @see GetPlayer, GetPlaylist
+	 * @see GetOverlays, GetPlayer, GetPlaylist
 	 */
 	int32 GetPlaylistIndex() const
 	{
@@ -643,7 +617,7 @@ public:
 	 * Get the sound wave that receives the audio output.
 	 *
 	 * @return Sound wave asset.
-	 * @see GetPlayer, GetVideoTexture
+	 * @see GetOverlays, GetPlayer, GetVideoTexture
 	 */
 	UMediaSoundWave* GetSoundWave() const
 	{
@@ -654,7 +628,7 @@ public:
 	 * Get the media texture that receives the video output.
 	 *
 	 * @return Media texture asset.
-	 * @see GetPlayer, GetSoundWave
+	 * @see GetOverlays, GetPlayer, GetSoundWave
 	 */
 	UMediaTexture* GetVideoTexture() const
 	{
@@ -683,44 +657,14 @@ public:
 	virtual void PreEditChange(UProperty* PropertyAboutToChange) override;
 #endif
 
-protected:
-
-	/**
-	 * Find a player that can play the specified media URL.
-	 *
-	 * @param Url The URL to play.
-	 * @param Options The media options for the URL.
-	 * @return The player if found, or nullptr otherwise.
-	 */
-	TSharedPtr<IMediaPlayer> FindPlayerForUrl(const FString& Url, const IMediaOptions& Options);
-
-	/**
-	 * Open a media source from a URL with optional parameters.
-	 *
-	 * @param Url The URL of the media to open (file name or web address).
-	 * @param Options Optional media parameters.
-	 * @return true if the media is being opened, false otherwise.
-	 */
-	bool Open(const FString& Url, const IMediaOptions& Options);
-
-	/** Select the default media tracks. */
-	void SelectDefaultTracks();
-
-protected:
+public:
 
-	//~ IMediaOverlaySink interface
+	//~ FTickerObjectBase interface
 
-	virtual bool InitializeOverlaySink() override;
-	virtual void AddOverlaySinkText(const FText& Text, EMediaOverlayType Type, FTimespan Time, FTimespan Duration, TOptional<FVector2D> Position) override;
-	virtual void ClearOverlaySinkText() override;
-	virtual void ShutdownOverlaySink() override;
+	virtual bool Tick(float DeltaTime) override;
 
 public:
 
-	/** Name of the desired native player, if any. */
-	UPROPERTY(transient)
-	FName DesiredPlayerName;
-
 	/** Automatically start playback after media opened successfully. */
 	UPROPERTY(EditAnywhere, Category=Playback)
 	bool PlayOnOpen;
@@ -729,12 +673,25 @@ public:
 	UPROPERTY(BlueprintReadOnly, EditAnywhere, Category=Playback)
 	uint32 Shuffle : 1;
 
+public:
+
+	/** Get an event delegate that is invoked when a media event occurred. */
+	DECLARE_EVENT_OneParam(UMediaPlayer, FOnMediaEvent, EMediaEvent /*Event*/)
+	FOnMediaEvent& OnMediaEvent()
+	{
+		return MediaEvent;
+	}
+
 protected:
 
 	/** Whether the player should loop when media playback reaches the end. */
 	UPROPERTY(BlueprintReadOnly, EditAnywhere, Category=Playback)
 	uint32 Loop:1;
 
+	/** The overlay asset to output captions and subtitles to. */
+	UPROPERTY(BlueprintReadOnly, EditAnywhere, Category=Output)
+	UMediaOverlays* Overlays;
+
 	/** The play list to use, if any. */
 	UPROPERTY(BlueprintReadOnly, transient, Category=Playback)
 	UMediaPlaylist* Playlist;
@@ -753,6 +710,9 @@ protected:
 
 private:
 
+	/** Callback for when a media overlays are being destroyed. */
+	void HandleMediaOverlaysBeginDestroy(UMediaOverlays& DestroyedOverlays);
+
 	/** Callback for when a media texture is being destroyed. */
 	void HandleMediaSoundWaveBeginDestroy(UMediaSoundWave& DestroyedSoundWave);
 
@@ -764,26 +724,14 @@ private:
 
 private:
 
-	/** Holds the URL of the currently loaded media source. */
-	FString CurrentUrl;
+	/** The URL of the media that was last attempted to be opened. */
+	FString LastUrl;
 
 	/** An event delegate that is invoked when a media event occurred. */
 	FOnMediaEvent MediaEvent;
 
-	struct FOverlay
-	{
-		FTimespan Duration;
-		TOptional<FVector2D> Position;
-		FText Text;
-		FTimespan Time;
-		EMediaOverlayType Type;
-	};
-
-	/** Queued text overlays. */
-	TArray<FOverlay> Overlays;
-
-	/** Holds the low-level player used to play the media source. */
-	TSharedPtr<IMediaPlayer> Player;
+	/** The player implementation. */
+	FMediaPlayerBase Player;
 
 #if WITH_EDITOR
 	/** Whether the player was playing in PIE/SIE. */
diff --git a/Engine/Source/Runtime/MediaAssets/Public/MediaSoundWave.h b/Engine/Source/Runtime/MediaAssets/Public/MediaSoundWave.h
--- a/Engine/Source/Runtime/MediaAssets/Public/MediaSoundWave.h
+++ b/Engine/Source/Runtime/MediaAssets/Public/MediaSoundWave.h
@@ -20,6 +20,11 @@ class MEDIAASSETS_API UMediaSoundWave
 
 public:
 
+	/**
+	 * Get the event delegate that is invoked when this asset is being destroyed.
+	 *
+	 * @return The delegate.
+	 */
 	DECLARE_EVENT_OneParam(UMediaSoundWave, FOnBeginDestroy, UMediaSoundWave& /*DestroyedSoundWave*/)
 	FOnBeginDestroy& OnBeginDestroy()
 	{
@@ -72,7 +77,7 @@ protected:
 
 private:
 
-	/** An event delegate that is invoked when this media texture is being destroyed. */
+	/** An event delegate that is invoked when this asset is being destroyed. */
 	FOnBeginDestroy BeginDestroyEvent;
 
 	/** Critical section for synchronizing access to QueuedAudio. */
diff --git a/Engine/Source/Runtime/MediaAssets/Public/MediaSource.h b/Engine/Source/Runtime/MediaAssets/Public/MediaSource.h
--- a/Engine/Source/Runtime/MediaAssets/Public/MediaSource.h
+++ b/Engine/Source/Runtime/MediaAssets/Public/MediaSource.h
@@ -6,8 +6,10 @@
 #include "UObject/ObjectMacros.h"
 #include "UObject/Object.h"
 #include "IMediaOptions.h"
+
 #include "MediaSource.generated.h"
 
+
 /**
  * Abstract base class for media sources.
  *
@@ -22,16 +24,6 @@ class MEDIAASSETS_API UMediaSource
 {
 	GENERATED_BODY()
 
-public:
-
-#if WITH_EDITORONLY_DATA
-
-	/** Override native media player plug-ins per platform (Empty = find one automatically). */
-	UPROPERTY(EditAnywhere, Category=Overrides, Meta=(DisplayName="Platform Players"))
-	TMap<FString, FName> PlatformPlayerNames;
-
-#endif
-
 public:
 
 	/** Get the media source's URL string (must override in child classes). */
@@ -40,17 +32,6 @@ public:
 	/** Validate the media source settings (must override in child classes). */
 	virtual bool Validate() const PURE_VIRTUAL(UMediaSource::Validate, return false;);
 
-public:
-
-	//~ UObject interface
-
-	virtual void GetAssetRegistryTags(TArray<FAssetRegistryTag>& OutTags) const override;
-	virtual void Serialize(FArchive& Ar) override;
-
-#if WITH_EDITOR
-	virtual void GetAssetRegistryTagMetadata(TMap<FName, FAssetRegistryTagMetadata>& OutMetadata) const override;
-#endif
-
 public:
 
 	//~ IMediaOptions interface
@@ -62,10 +43,4 @@ public:
 	virtual FString GetMediaOption(const FName& Key, const FString& DefaultValue) const override;
 	virtual FText GetMediaOption(const FName& Key, const FText& DefaultValue) const override;
 	virtual bool HasMediaOption(const FName& Key) const override;
-
-private:
-
-	/** Name of the desired native media player (Empty = find one automatically). */
-	UPROPERTY()
-	FName PlayerName;
 };
diff --git a/Engine/Source/Runtime/MediaAssets/Public/MediaTexture.h b/Engine/Source/Runtime/MediaAssets/Public/MediaTexture.h
--- a/Engine/Source/Runtime/MediaAssets/Public/MediaTexture.h
+++ b/Engine/Source/Runtime/MediaAssets/Public/MediaTexture.h
@@ -64,29 +64,17 @@ public:
 
 public:
 
+	/**
+	 * Get the event delegate that is invoked when this asset is being destroyed.
+	 *
+	 * @return The delegate.
+	 */
 	DECLARE_EVENT_OneParam(UMediaTexture, FOnBeginDestroy, UMediaTexture& /*DestroyedMediaTexture*/)
 	FOnBeginDestroy& OnBeginDestroy()
 	{
 		return BeginDestroyEvent;
 	}
 
-public:
-
-	//~ IMediaTextureSink interface
-
-	virtual void* AcquireTextureSinkBuffer() override;
-	virtual void DisplayTextureSinkBuffer(FTimespan Time) override;
-	virtual FIntPoint GetTextureSinkDimensions() const override;
-	virtual EMediaTextureSinkFormat GetTextureSinkFormat() const override;
-	virtual EMediaTextureSinkMode GetTextureSinkMode() const override;
-	virtual FRHITexture* GetTextureSinkTexture() override;
-	virtual bool InitializeTextureSink(FIntPoint OutputDim, FIntPoint BufferDim, EMediaTextureSinkFormat Format, EMediaTextureSinkMode Mode) override;
-	virtual void ReleaseTextureSinkBuffer() override;
-	virtual void ShutdownTextureSink() override;
-	virtual bool SupportsTextureSinkFormat(EMediaTextureSinkFormat Format) const override;
-	virtual void UpdateTextureSinkBuffer(const uint8* Data, uint32 Pitch = 0) override;
-	virtual void UpdateTextureSinkResource(FRHITexture* RenderTarget, FRHITexture* ShaderResource) override;
-
 public:
 
 	//~ UTexture interface.
@@ -110,6 +98,23 @@ public:
 	virtual void PreEditChange(UProperty* PropertyAboutToChange) override;
 #endif
 
+public:
+
+	//~ IMediaTextureSink interface
+
+	virtual void* AcquireTextureSinkBuffer() override;
+	virtual void DisplayTextureSinkBuffer(FTimespan Time) override;
+	virtual FIntPoint GetTextureSinkDimensions() const override;
+	virtual EMediaTextureSinkFormat GetTextureSinkFormat() const override;
+	virtual EMediaTextureSinkMode GetTextureSinkMode() const override;
+	virtual FRHITexture* GetTextureSinkTexture() override;
+	virtual bool InitializeTextureSink(FIntPoint OutputDim, FIntPoint BufferDim, EMediaTextureSinkFormat Format, EMediaTextureSinkMode Mode) override;
+	virtual void ReleaseTextureSinkBuffer() override;
+	virtual void ShutdownTextureSink() override;
+	virtual bool SupportsTextureSinkFormat(EMediaTextureSinkFormat Format) const override;
+	virtual void UpdateTextureSinkBuffer(const uint8* Data, uint32 Pitch = 0) override;
+	virtual void UpdateTextureSinkResource(FRHITexture* RenderTarget, FRHITexture* ShaderResource) override;
+
 protected:
 
 	//~ Deprecated members
@@ -124,7 +129,7 @@ protected:
 
 private:
 
-	/** An event delegate that is invoked when this media texture is being destroyed. */
+	/** An event delegate that is invoked when this asset is being destroyed. */
 	FOnBeginDestroy BeginDestroyEvent;
 
 	/** Critical section for synchronizing access to texture resource object. */
diff --git a/Engine/Source/Runtime/MediaAssets/Public/PlatformMediaSource.h b/Engine/Source/Runtime/MediaAssets/Public/PlatformMediaSource.h
--- a/Engine/Source/Runtime/MediaAssets/Public/PlatformMediaSource.h
+++ b/Engine/Source/Runtime/MediaAssets/Public/PlatformMediaSource.h
@@ -5,14 +5,16 @@
 #include "CoreMinimal.h"
 #include "UObject/ObjectMacros.h"
 #include "MediaSource.h"
+
 #include "PlatformMediaSource.generated.h"
 
+
 /**
  * A media source that selects other media sources based on target platform.
  *
  * Use this asset to override media sources on a per-platform basis.
  */
-UCLASS(BlueprintType, hidecategories=(Overrides))
+UCLASS(BlueprintType)
 class MEDIAASSETS_API UPlatformMediaSource
 	: public UMediaSource
 {
diff --git a/Engine/Source/Runtime/MediaAssets/Public/StreamMediaSource.h b/Engine/Source/Runtime/MediaAssets/Public/StreamMediaSource.h
--- a/Engine/Source/Runtime/MediaAssets/Public/StreamMediaSource.h
+++ b/Engine/Source/Runtime/MediaAssets/Public/StreamMediaSource.h
@@ -4,13 +4,14 @@
 
 #include "CoreMinimal.h"
 #include "UObject/ObjectMacros.h"
-#include "MediaSource.h"
+#include "BaseMediaSource.h"
+
 #include "StreamMediaSource.generated.h"
 
 
 UCLASS(BlueprintType)
 class MEDIAASSETS_API UStreamMediaSource
-	: public UMediaSource
+	: public UBaseMediaSource
 {
 	GENERATED_BODY()
 
diff --git a/Engine/Source/Runtime/MediaUtils/MediaUtils.Build.cs b/Engine/Source/Runtime/MediaUtils/MediaUtils.Build.cs
--- a/Engine/Source/Runtime/MediaUtils/MediaUtils.Build.cs
+++ b/Engine/Source/Runtime/MediaUtils/MediaUtils.Build.cs
@@ -10,6 +10,7 @@ namespace UnrealBuildTool.Rules
 				new string[]
 				{
 					"Core",
+                    "Media",
 				}
 			); 
 
diff --git a/Engine/Source/Runtime/MediaUtils/Private/MediaPlayerBase.cpp b/Engine/Source/Runtime/MediaUtils/Private/MediaPlayerBase.cpp
new file mode 100644
--- /dev/null
+++ b/Engine/Source/Runtime/MediaUtils/Private/MediaPlayerBase.cpp
@@ -0,0 +1,468 @@
+// Copyright 1998-2016 Epic Games, Inc. All Rights Reserved.
+
+#include "MediaPlayerBase.h"
+
+#include "IMediaControls.h"
+#include "IMediaModule.h"
+#include "IMediaOptions.h"
+#include "IMediaOutput.h"
+#include "IMediaPlayer.h"
+#include "IMediaPlayerFactory.h"
+#include "IMediaTracks.h"
+#include "MediaUtilsPrivate.h"
+#include "Misc/CoreMisc.h"
+#include "Misc/ScopeLock.h"
+#include "Modules/ModuleManager.h"
+
+
+/* FMediaPlayerBase structors
+*****************************************************************************/
+
+FMediaPlayerBase::~FMediaPlayerBase()
+{
+	if (NativePlayer.IsValid())
+	{
+		FScopeLock Lock(&CriticalSection);
+
+		NativePlayer->Close();
+		NativePlayer->GetOutput().SetOverlaySink(nullptr);
+		NativePlayer->OnMediaEvent().RemoveAll(this);
+		NativePlayer.Reset();
+	}
+}
+
+
+/* FMediaPlayerBase interface
+*****************************************************************************/
+
+bool FMediaPlayerBase::CanPause() const
+{
+	if (!NativePlayer.IsValid())
+	{
+		return false;
+	}
+
+	EMediaState State = NativePlayer->GetControls().GetState();
+
+	return ((State == EMediaState::Playing) || (State == EMediaState::Preparing));
+}
+
+
+bool FMediaPlayerBase::CanPlayUrl(const FString& Url, const IMediaOptions& Options)
+{
+	return FindPlayerForUrl(Url, Options).IsValid();
+}
+
+
+void FMediaPlayerBase::Close()
+{
+	if (CurrentUrl.IsEmpty())
+	{
+		return;
+	}
+
+	CurrentUrl.Empty();
+
+	if (NativePlayer.IsValid())
+	{
+		NativePlayer->Close();
+	}
+}
+
+
+FTimespan FMediaPlayerBase::GetDuration() const
+{
+	return NativePlayer.IsValid() ? NativePlayer->GetControls().GetDuration() : FTimespan::Zero();
+}
+
+
+TRange<float> FMediaPlayerBase::GetForwardRates(bool Unthinned)
+{
+	return NativePlayer.IsValid()
+		? NativePlayer->GetControls().GetSupportedRates(EMediaPlaybackDirections::Forward, Unthinned)
+		: TRange<float>::Empty();
+}
+
+
+FString FMediaPlayerBase::GetInfo() const
+{
+	return NativePlayer.IsValid() ? NativePlayer->GetInfo() : FString();
+}
+
+
+int32 FMediaPlayerBase::GetNumTracks(EMediaTrackType TrackType) const
+{
+	return NativePlayer.IsValid() ? NativePlayer->GetTracks().GetNumTracks(TrackType) : 0;
+}
+
+
+FName FMediaPlayerBase::GetPlayerName() const
+{
+	return NativePlayer.IsValid() ? NativePlayer->GetName() : NAME_None;
+}
+
+
+float FMediaPlayerBase::GetRate() const
+{
+	return NativePlayer.IsValid() ? NativePlayer->GetControls().GetRate() : 0.0f;
+}
+
+
+TRange<float> FMediaPlayerBase::GetReverseRates(bool Unthinned)
+{
+	return NativePlayer.IsValid()
+		? NativePlayer->GetControls().GetSupportedRates(EMediaPlaybackDirections::Reverse, Unthinned)
+		: TRange<float>::Empty();
+}
+
+
+int32 FMediaPlayerBase::GetSelectedTrack(EMediaTrackType TrackType) const
+{
+	return NativePlayer.IsValid() ? NativePlayer->GetTracks().GetSelectedTrack((EMediaTrackType)TrackType) : INDEX_NONE;
+}
+
+
+FString FMediaPlayerBase::GetStats() const
+{
+	return NativePlayer.IsValid() ? NativePlayer->GetStats() : FString();
+}
+
+
+FTimespan FMediaPlayerBase::GetTime() const
+{
+	return NativePlayer.IsValid() ? NativePlayer->GetControls().GetTime() : FTimespan::Zero();
+}
+
+
+FText FMediaPlayerBase::GetTrackDisplayName(EMediaTrackType TrackType, int32 TrackIndex) const
+{
+	return NativePlayer.IsValid() ? NativePlayer->GetTracks().GetTrackDisplayName((EMediaTrackType)TrackType, TrackIndex) : FText::GetEmpty();
+}
+
+
+FString FMediaPlayerBase::GetTrackLanguage(EMediaTrackType TrackType, int32 TrackIndex) const
+{
+	return NativePlayer.IsValid() ? NativePlayer->GetTracks().GetTrackLanguage((EMediaTrackType)TrackType, TrackIndex) : FString();
+}
+
+
+bool FMediaPlayerBase::IsLooping() const
+{
+	return NativePlayer.IsValid() && NativePlayer->GetControls().IsLooping();
+}
+
+
+bool FMediaPlayerBase::IsPaused() const
+{
+	return NativePlayer.IsValid() && (NativePlayer->GetControls().GetState() == EMediaState::Paused);
+}
+
+
+bool FMediaPlayerBase::IsPlaying() const
+{
+	return NativePlayer.IsValid() && (NativePlayer->GetControls().GetState() == EMediaState::Playing);
+}
+
+
+bool FMediaPlayerBase::IsPreparing() const
+{
+	return NativePlayer.IsValid() && (NativePlayer->GetControls().GetState() == EMediaState::Preparing);
+}
+
+
+bool FMediaPlayerBase::IsReady() const
+{
+	if (!NativePlayer.IsValid())
+	{
+		return false;
+	}
+
+	return ((NativePlayer->GetControls().GetState() != EMediaState::Closed) &&
+		(NativePlayer->GetControls().GetState() != EMediaState::Error) &&
+		(NativePlayer->GetControls().GetState() != EMediaState::Preparing));
+}
+
+
+bool FMediaPlayerBase::Open(const FString& Url, const IMediaOptions& Options)
+{
+	if (IsRunningDedicatedServer())
+	{
+		return false;
+	}
+
+	Close();
+
+	if (Url.IsEmpty())
+	{
+		return false;
+	}
+
+	// find & initialize new player
+	TSharedPtr<IMediaPlayer> NewPlayer = FindPlayerForUrl(Url, Options);
+
+	if (NewPlayer != NativePlayer)
+	{
+		if (NativePlayer.IsValid())
+		{
+			NativePlayer->OnMediaEvent().RemoveAll(this);
+		}
+
+		if (NewPlayer.IsValid())
+		{
+			NewPlayer->OnMediaEvent().AddRaw(this, &FMediaPlayerBase::HandlePlayerMediaEvent);
+		}
+
+		FScopeLock Lock(&CriticalSection);
+		NativePlayer = NewPlayer;
+	}
+
+	if (!NativePlayer.IsValid())
+	{
+		return false;
+	}
+
+	CurrentUrl = Url;
+
+	// open the new media source
+	if (!NativePlayer->Open(Url, Options))
+	{
+		CurrentUrl.Empty();
+
+		return false;
+	}
+
+	return true;
+}
+
+
+bool FMediaPlayerBase::Seek(const FTimespan& Time)
+{
+	return NativePlayer.IsValid() && NativePlayer->GetControls().Seek(Time);
+}
+
+
+bool FMediaPlayerBase::SelectTrack(EMediaTrackType TrackType, int32 TrackIndex)
+{
+	return NativePlayer.IsValid() ? NativePlayer->GetTracks().SelectTrack((EMediaTrackType)TrackType, TrackIndex) : false;
+}
+
+
+void FMediaPlayerBase::SetAudioSink(IMediaAudioSink* Sink)
+{
+	if (NativePlayer.IsValid())
+	{
+		NativePlayer->GetOutput().SetAudioSink(Sink);
+	}
+}
+
+
+bool FMediaPlayerBase::SetLooping(bool Looping)
+{
+	return NativePlayer.IsValid() && NativePlayer->GetControls().SetLooping(Looping);
+}
+
+
+void FMediaPlayerBase::SetMetadataSink(IMediaBinarySink* Sink)
+{
+	if (NativePlayer.IsValid())
+	{
+		NativePlayer->GetOutput().SetMetadataSink(Sink);
+	}
+}
+
+
+void FMediaPlayerBase::SetOverlaySink(IMediaOverlaySink* Sink)
+{
+	if (NativePlayer.IsValid())
+	{
+		NativePlayer->GetOutput().SetOverlaySink(Sink);
+	}
+}
+
+
+bool FMediaPlayerBase::SetRate(float Rate)
+{
+	return NativePlayer.IsValid() && NativePlayer->GetControls().SetRate(Rate);
+}
+
+
+void FMediaPlayerBase::SetVideoSink(IMediaTextureSink* Sink)
+{
+	if (NativePlayer.IsValid())
+	{
+		NativePlayer->GetOutput().SetVideoSink(Sink);
+	}
+}
+
+
+bool FMediaPlayerBase::SupportsRate(float Rate, bool Unthinned) const
+{
+	return NativePlayer.IsValid() && NativePlayer->GetControls().SupportsRate(Rate, Unthinned);
+}
+
+
+bool FMediaPlayerBase::SupportsScrubbing() const
+{
+	return NativePlayer.IsValid() && NativePlayer->GetControls().SupportsScrubbing();
+}
+
+
+bool FMediaPlayerBase::SupportsSeeking() const
+{
+	return NativePlayer.IsValid() && NativePlayer->GetControls().SupportsSeeking();
+}
+
+
+void FMediaPlayerBase::TickPlayer(float DeltaTime)
+{
+	if (NativePlayer.IsValid())
+	{
+		NativePlayer->TickPlayer(DeltaTime);
+	}
+}
+
+
+void FMediaPlayerBase::TickVideo(float DeltaTime)
+{
+	FScopeLock Lock(&CriticalSection);
+
+	if (NativePlayer.IsValid())
+	{
+		NativePlayer->TickVideo(DeltaTime);
+	}
+}
+
+
+/* FMediaPlayerBase implementation
+*****************************************************************************/
+
+TSharedPtr<IMediaPlayer> FMediaPlayerBase::FindPlayerForUrl(const FString& Url, const IMediaOptions& Options)
+{
+	const FName PlayerName = (DesiredPlayerName != NAME_None) ? DesiredPlayerName : Options.GetDesiredPlayerName();
+
+	// reuse existing player if desired
+	if (NativePlayer.IsValid() && (PlayerName == NativePlayer->GetName()))
+	{
+		return NativePlayer;
+	}
+
+	// load media module
+	IMediaModule* MediaModule = FModuleManager::LoadModulePtr<IMediaModule>("Media");
+
+	if (MediaModule == nullptr)
+	{
+		UE_LOG(LogMediaUtils, Error, TEXT("Failed to load Media module"));
+		return nullptr;
+	}
+
+	// try to create desired player
+	if (PlayerName != NAME_None)
+	{
+		IMediaPlayerFactory* Factory = MediaModule->GetPlayerFactory(PlayerName);
+
+		if (Factory == nullptr)
+		{
+			UE_LOG(LogMediaUtils, Error, TEXT("Could not find desired player %s for %s"), *PlayerName.ToString(), *Url);
+			return nullptr;
+		}
+
+		TSharedPtr<IMediaPlayer> NewPlayer = Factory->CreatePlayer();
+
+		if (!NewPlayer.IsValid())
+		{
+			UE_LOG(LogMediaUtils, Error, TEXT("Failed to create desired player %s for %s"), *PlayerName.ToString(), *Url);
+			return nullptr;
+		}
+
+		return NewPlayer;
+	}
+
+	// try to reuse existing player
+	if (NativePlayer.IsValid())
+	{
+		IMediaPlayerFactory* Factory = MediaModule->GetPlayerFactory(NativePlayer->GetName());
+
+		if ((Factory != nullptr) && Factory->CanPlayUrl(Url, Options))
+		{
+			return NativePlayer;
+		}
+	}
+
+	const FString RunningPlatformName(FPlatformProperties::IniPlatformName());
+
+	// try to auto-select new player
+	const TArray<IMediaPlayerFactory*>& PlayerFactories = MediaModule->GetPlayerFactories();
+
+	for (IMediaPlayerFactory* Factory : PlayerFactories)
+	{
+		if (!Factory->SupportsPlatform(RunningPlatformName) || !Factory->CanPlayUrl(Url, Options))
+		{
+			continue;
+		}
+
+		TSharedPtr<IMediaPlayer> NewPlayer = Factory->CreatePlayer();
+
+		if (NewPlayer.IsValid())
+		{
+			return NewPlayer;
+		}
+	}
+
+	// no suitable player found
+	if (PlayerFactories.Num() > 0)
+	{
+		UE_LOG(LogMediaUtils, Error, TEXT("Cannot play %s, because none of the enabled media player plug-ins support it:"), *Url);
+
+		for (IMediaPlayerFactory* Factory : PlayerFactories)
+		{
+			if (Factory->SupportsPlatform(RunningPlatformName))
+			{
+				UE_LOG(LogMediaUtils, Error, TEXT("| %s (URI scheme or file extension not supported)"), *Factory->GetName().ToString());
+			}
+			else
+			{
+				UE_LOG(LogMediaUtils, Error, TEXT("| %s (only available on %s, but not on %s)"), *Factory->GetName().ToString(), *FString::Join(Factory->GetSupportedPlatforms(), TEXT(", ")), *RunningPlatformName);
+			}	
+		}
+	}
+	else
+	{
+		UE_LOG(LogMediaUtils, Error, TEXT("Cannot play %s: no media player plug-ins are installed and enabled in this project"), *Url);
+	}
+
+	return nullptr;
+}
+
+
+void FMediaPlayerBase::SelectDefaultTracks()
+{
+	IMediaTracks& Tracks = NativePlayer->GetTracks();
+
+	// @todo Media: consider locale when selecting default tracks
+	Tracks.SelectTrack(EMediaTrackType::Audio, 0);
+	Tracks.SelectTrack(EMediaTrackType::Caption, 0);
+	Tracks.SelectTrack(EMediaTrackType::Video, 0);
+}
+
+
+/* FMediaPlayerBase callbacks
+*****************************************************************************/
+
+void FMediaPlayerBase::HandlePlayerMediaEvent(EMediaEvent Event)
+{
+	if (Event == EMediaEvent::TracksChanged)
+	{
+		SelectDefaultTracks();
+	}
+	else if ((Event == EMediaEvent::MediaOpened) || (Event == EMediaEvent::MediaOpenFailed))
+	{
+		if (Event == EMediaEvent::MediaOpenFailed)
+		{
+			CurrentUrl.Empty();
+		}
+
+		UE_LOG(LogMediaUtils, Verbose, TEXT("Media Info:\n\n%s"), *NativePlayer->GetInfo());
+	}
+
+	MediaEvent.Broadcast(Event);
+}
diff --git a/Engine/Source/Runtime/MediaUtils/Private/MediaUtilsModule.cpp b/Engine/Source/Runtime/MediaUtils/Private/MediaUtilsModule.cpp
--- a/Engine/Source/Runtime/MediaUtils/Private/MediaUtilsModule.cpp
+++ b/Engine/Source/Runtime/MediaUtils/Private/MediaUtilsModule.cpp
@@ -1,9 +1,14 @@
 // Copyright 1998-2017 Epic Games, Inc. All Rights Reserved.
 
+#include "MediaUtilsPrivate.h"
+
 #include "Modules/ModuleInterface.h"
 #include "Modules/ModuleManager.h"
 
 
+DEFINE_LOG_CATEGORY(LogMediaUtils);
+
+
 /**
  * Implements the Media module.
  */
diff --git a/Engine/Source/Runtime/MediaUtils/Private/MediaUtilsPrivate.h b/Engine/Source/Runtime/MediaUtils/Private/MediaUtilsPrivate.h
new file mode 100644
--- /dev/null
+++ b/Engine/Source/Runtime/MediaUtils/Private/MediaUtilsPrivate.h
@@ -0,0 +1,9 @@
+// Copyright 1998-2016 Epic Games, Inc. All Rights Reserved.
+
+#pragma once
+
+#include "CoreMinimal.h"
+
+
+/** Declares a log category for this module. */
+DECLARE_LOG_CATEGORY_EXTERN(LogMediaUtils, Log, All);
