MediaAssets: Fixed media player crash on garbage collection

This is a temporary fix that may introduce a hitch when discarding media players.
However, this entire code path will be removed entirely for 4.16.

diff --git a/Engine/Source/Runtime/MediaAssets/Private/Assets/MediaPlayer.cpp b/Engine/Source/Runtime/MediaAssets/Private/Assets/MediaPlayer.cpp
--- a/Engine/Source/Runtime/MediaAssets/Private/Assets/MediaPlayer.cpp
+++ b/Engine/Source/Runtime/MediaAssets/Private/Assets/MediaPlayer.cpp
@@ -493,6 +493,7 @@ void UMediaPlayer::BeginDestroy()
 {
 	Super::BeginDestroy();
 
+	FlushRenderingCommands();
 	Close();
 
 	SetOverlays(nullptr);
