MediaAssets: Fixed crash in TickVideo after media player is GC'ed

Copied from CL# 3299912 & 3300133

diff --git a/Engine/Source/Runtime/MediaAssets/Private/Assets/MediaPlayer.cpp b/Engine/Source/Runtime/MediaAssets/Private/Assets/MediaPlayer.cpp
--- a/Engine/Source/Runtime/MediaAssets/Private/Assets/MediaPlayer.cpp
+++ b/Engine/Source/Runtime/MediaAssets/Private/Assets/MediaPlayer.cpp
@@ -27,11 +27,12 @@ UMediaPlayer::UMediaPlayer(const FObjectInitializer& ObjectInitializer)
 	, Shuffle(false)
 	, Loop(false)
 	, PlaylistIndex(INDEX_NONE)
+	, Player(MakeShareable(new FMediaPlayerBase))
 #if WITH_EDITOR
 	, WasPlayingInPIE(false)
 #endif
 {
-	Player.OnMediaEvent().AddUObject(this, &UMediaPlayer::HandlePlayerMediaEvent);
+	Player->OnMediaEvent().AddUObject(this, &UMediaPlayer::HandlePlayerMediaEvent);
 }
 
 
@@ -40,7 +41,7 @@ UMediaPlayer::UMediaPlayer(const FObjectInitializer& ObjectInitializer)
 
 bool UMediaPlayer::CanPause() const
 {
-	return Player.CanPause();
+	return Player->CanPause();
 }
 
 
@@ -51,7 +52,7 @@ bool UMediaPlayer::CanPlaySource(UMediaSource* MediaSource)
 		return false;
 	}
 
-	return Player.CanPlayUrl(MediaSource->GetUrl(), *MediaSource);
+	return Player->CanPlayUrl(MediaSource->GetUrl(), *MediaSource);
 }
 
 
@@ -62,13 +63,13 @@ bool UMediaPlayer::CanPlayUrl(const FString& Url)
 		return false;
 	}
 
-	return Player.CanPlayUrl(Url, *GetDefault<UMediaSource>());
+	return Player->CanPlayUrl(Url, *GetDefault<UMediaSource>());
 }
 
 
 void UMediaPlayer::Close()
 {
-	Player.Close();
+	Player->Close();
 
 	LastUrl.Empty();
 	Playlist = nullptr;
@@ -78,109 +79,109 @@ void UMediaPlayer::Close()
 
 FMediaPlayerBase& UMediaPlayer::GetBasePlayer()
 {
-	return Player;
+	return *Player;
 }
 
 
 FName UMediaPlayer::GetDesiredPlayerName() const
 {
-	return Player.DesiredPlayerName;
+	return Player->DesiredPlayerName;
 }
 
 
 FTimespan UMediaPlayer::GetDuration() const
 {
-	return Player.GetDuration();
+	return Player->GetDuration();
 }
 
 
 FFloatRange UMediaPlayer::GetForwardRates(bool Unthinned)
 {
-	return Player.GetForwardRates(Unthinned);
+	return Player->GetForwardRates(Unthinned);
 }
 
 
 int32 UMediaPlayer::GetNumTracks(EMediaPlayerTrack TrackType) const
 {
-	return Player.GetNumTracks((EMediaTrackType)TrackType);
+	return Player->GetNumTracks((EMediaTrackType)TrackType);
 }
 
 
 FName UMediaPlayer::GetPlayerName() const
 {
-	return Player.GetPlayerName();
+	return Player->GetPlayerName();
 }
 
 
 float UMediaPlayer::GetRate() const
 {
-	return Player.GetRate();
+	return Player->GetRate();
 }
 
 
 FFloatRange UMediaPlayer::GetReverseRates(bool Unthinned)
 {
-	return Player.GetReverseRates(Unthinned);
+	return Player->GetReverseRates(Unthinned);
 }
 
 
 int32 UMediaPlayer::GetSelectedTrack(EMediaPlayerTrack TrackType) const
 {
-	return Player.GetSelectedTrack((EMediaTrackType)TrackType);
+	return Player->GetSelectedTrack((EMediaTrackType)TrackType);
 }
 
 
 FTimespan UMediaPlayer::GetTime() const
 {
-	return Player.GetTime();
+	return Player->GetTime();
 }
 
 
 FText UMediaPlayer::GetTrackDisplayName(EMediaPlayerTrack TrackType, int32 TrackIndex) const
 {
-	return Player.GetTrackDisplayName((EMediaTrackType)TrackType, TrackIndex);
+	return Player->GetTrackDisplayName((EMediaTrackType)TrackType, TrackIndex);
 }
 
 
 FString UMediaPlayer::GetTrackLanguage(EMediaPlayerTrack TrackType, int32 TrackIndex) const
 {
-	return Player.GetTrackLanguage((EMediaTrackType)TrackType, TrackIndex);
+	return Player->GetTrackLanguage((EMediaTrackType)TrackType, TrackIndex);
 }
 
 
 const FString& UMediaPlayer::GetUrl() const
 {
-	return Player.GetUrl();
+	return Player->GetUrl();
 }
 
 
 bool UMediaPlayer::IsLooping() const
 {
-	return Player.IsLooping();
+	return Player->IsLooping();
 }
 
 
 bool UMediaPlayer::IsPaused() const
 {
-	return Player.IsPaused();
+	return Player->IsPaused();
 }
 
 
 bool UMediaPlayer::IsPlaying() const
 {
-	return Player.IsPlaying();
+	return Player->IsPlaying();
 }
 
 
 bool UMediaPlayer::IsPreparing() const
 {
-	return Player.IsPreparing();
+	return Player->IsPreparing();
 }
 
 
 bool UMediaPlayer::IsReady() const
 {
-	return Player.IsReady();
+	return Player->IsReady();
 }
 
 
@@ -199,7 +200,7 @@ bool UMediaPlayer::Next()
 			? Playlist->GetRandom(PlaylistIndex)
 			: Playlist->GetNext(PlaylistIndex);
 
-		if ((NextSource != nullptr) && NextSource->Validate() && Player.Open(NextSource->GetUrl(), *NextSource))
+		if ((NextSource != nullptr) && NextSource->Validate() && Player->Open(NextSource->GetUrl(), *NextSource))
 		{
 			return true;
 		}
@@ -262,7 +263,7 @@ bool UMediaPlayer::OpenPlaylistIndex(UMediaPlaylist* InPlaylist, int32 Index)
 		return false;
 	}
 
-	return Player.Open(LastUrl, *MediaSource);
+	return Player->Open(LastUrl, *MediaSource);
 }
 
 
@@ -285,7 +286,7 @@ bool UMediaPlayer::OpenSource(UMediaSource* MediaSource)
 		return false;
 	}
 
-	return Player.Open(LastUrl, *MediaSource);
+	return Player->Open(LastUrl, *MediaSource);
 }
 
 
@@ -300,19 +301,19 @@ bool UMediaPlayer::OpenUrl(const FString& Url)
 
 	LastUrl = Url;
 
-	return Player.Open(LastUrl, *GetDefault<UMediaSource>());
+	return Player->Open(LastUrl, *GetDefault<UMediaSource>());
 }
 
 
 bool UMediaPlayer::Pause()
 {
-	return Player.SetRate(0.0f);
+	return Player->SetRate(0.0f);
 }
 
 
 bool UMediaPlayer::Play()
 {
-	return Player.SetRate(1.0f);
+	return Player->SetRate(1.0f);
 }
 
 
@@ -331,7 +332,7 @@ bool UMediaPlayer::Previous()
 			? Playlist->GetRandom(PlaylistIndex)
 			: Playlist->GetNext(PlaylistIndex);
 
-		if ((PrevSource != nullptr) && PrevSource->Validate() && Player.Open(PrevSource->GetUrl(), *PrevSource))
+		if ((PrevSource != nullptr) && PrevSource->Validate() && Player->Open(PrevSource->GetUrl(), *PrevSource))
 		{
 			return true;
 		}
@@ -348,7 +349,7 @@ bool UMediaPlayer::Reopen()
 		return OpenPlaylistIndex(Playlist, PlaylistIndex);
 	}
 
-	return Player.Open(LastUrl, *GetDefault<UMediaSource>());
+	return Player->Open(LastUrl, *GetDefault<UMediaSource>());
 }
 
 
@@ -360,19 +361,19 @@ bool UMediaPlayer::Rewind()
 
 bool UMediaPlayer::Seek(const FTimespan& Time)
 {
-	return Player.Seek(Time);
+	return Player->Seek(Time);
 }
 
 
 bool UMediaPlayer::SelectTrack(EMediaPlayerTrack TrackType, int32 TrackIndex)
 {
-	return Player.SelectTrack((EMediaTrackType)TrackType, TrackIndex);
+	return Player->SelectTrack((EMediaTrackType)TrackType, TrackIndex);
 }
 
 
 void UMediaPlayer::SetDesiredPlayerName(FName PlayerName)
 {
-	Player.DesiredPlayerName = PlayerName;
+	Player->DesiredPlayerName = PlayerName;
 }
 
 
@@ -380,7 +381,7 @@ bool UMediaPlayer::SetLooping(bool Looping)
 {
 	Loop = Looping;
 
-	return Player.SetLooping(Looping);
+	return Player->SetLooping(Looping);
 }
 
 
@@ -396,7 +397,7 @@ void UMediaPlayer::SetOverlays(UMediaOverlays* NewOverlays)
 		NewOverlays->OnBeginDestroy().AddUObject(this, &UMediaPlayer::HandleMediaOverlaysBeginDestroy);
 	}
 
-	Player.SetOverlaySink(NewOverlays);
+	Player->SetOverlaySink(NewOverlays);
 
 	Overlays = NewOverlays;
 }
@@ -404,7 +405,7 @@ void UMediaPlayer::SetOverlays(UMediaOverlays* NewOverlays)
 
 bool UMediaPlayer::SetRate(float Rate)
 {
-	return Player.SetRate(Rate);
+	return Player->SetRate(Rate);
 }
 
 
@@ -420,7 +421,7 @@ void UMediaPlayer::SetSoundWave(UMediaSoundWave* NewSoundWave)
 		NewSoundWave->OnBeginDestroy().AddUObject(this, &UMediaPlayer::HandleMediaSoundWaveBeginDestroy);
 	}
 
-	Player.SetAudioSink(NewSoundWave);
+	Player->SetAudioSink(NewSoundWave);
 
 	SoundWave = NewSoundWave;
 }
@@ -438,7 +439,7 @@ void UMediaPlayer::SetVideoTexture(UMediaTexture* NewTexture)
 		NewTexture->OnBeginDestroy().AddUObject(this, &UMediaPlayer::HandleMediaTextureBeginDestroy);
 	}
 
-	Player.SetVideoSink(NewTexture);
+	Player->SetVideoSink(NewTexture);
 
 	VideoTexture = NewTexture;
 }
@@ -446,19 +447,19 @@ void UMediaPlayer::SetVideoTexture(UMediaTexture* NewTexture)
 
 bool UMediaPlayer::SupportsRate(float Rate, bool Unthinned) const
 {
-	return Player.SupportsRate(Rate, Unthinned);
+	return Player->SupportsRate(Rate, Unthinned);
 }
 
 
 bool UMediaPlayer::SupportsScrubbing() const
 {
-	return Player.SupportsScrubbing();
+	return Player->SupportsScrubbing();
 }
 
 
 bool UMediaPlayer::SupportsSeeking() const
 {
-	return Player.SupportsSeeking();
+	return Player->SupportsSeeking();
 }
 
 
@@ -588,13 +589,19 @@ void UMediaPlayer::PreEditChange(UProperty* PropertyAboutToChange)
 
 bool UMediaPlayer::Tick(float DeltaTime)
 {
-	Player.TickPlayer(DeltaTime);
+	Player->TickPlayer(DeltaTime);
+
+	typedef TWeakPtr<FMediaPlayerBase, ESPMode::ThreadSafe> FMediaPlayerBasePtr;
 
 	ENQUEUE_UNIQUE_RENDER_COMMAND_TWOPARAMETER(MediaPlayerTickRender,
-		UMediaPlayer*, This, this,
+		FMediaPlayerBasePtr, PlayerPtr, Player,
 		float, DeltaTime, DeltaTime,
 		{
-			This->Player.TickVideo(DeltaTime);
+			auto PinnedPlayer = PlayerPtr.Pin();
+			if (PinnedPlayer.IsValid())
+			{
+				PinnedPlayer->TickVideo(DeltaTime);
+			}
 		});
 
 	return true;
@@ -608,7 +615,7 @@ void UMediaPlayer::HandleMediaOverlaysBeginDestroy(UMediaOverlays& DestroyedOver
 {
 	if (&DestroyedOverlays == Overlays)
 	{
-		Player.SetOverlaySink(nullptr);
+		Player->SetOverlaySink(nullptr);
 	}
 }
 
@@ -617,7 +624,7 @@ void UMediaPlayer::HandleMediaSoundWaveBeginDestroy(UMediaSoundWave& DestroyedSo
 {
 	if (&DestroyedSoundWave == SoundWave)
 	{
-		Player.SetAudioSink(nullptr);
+		Player->SetAudioSink(nullptr);
 	}
 }
 
@@ -626,7 +633,7 @@ void UMediaPlayer::HandleMediaTextureBeginDestroy(UMediaTexture& DestroyedMediaT
 {
 	if (&DestroyedMediaTexture == VideoTexture)
 	{
-		Player.SetVideoSink(nullptr);
+		Player->SetVideoSink(nullptr);
 	}
 }
 
@@ -642,24 +649,24 @@ void UMediaPlayer::HandlePlayerMediaEvent(EMediaEvent Event)
 		break;
 
 	case EMediaEvent::MediaOpened:
-		Player.SetLooping(Loop);
+		Player->SetLooping(Loop);
 
 		if (Overlays != nullptr)
 		{
-			Player.SetOverlaySink(Overlays);
+			Player->SetOverlaySink(Overlays);
 		}
 
 		if (SoundWave != nullptr)
 		{
-			Player.SetAudioSink(SoundWave);
+			Player->SetAudioSink(SoundWave);
 		}
 
 		if (VideoTexture != nullptr)
 		{
-			Player.SetVideoSink(VideoTexture);
+			Player->SetVideoSink(VideoTexture);
 		}
 
-		OnMediaOpened.Broadcast(Player.GetUrl());
+		OnMediaOpened.Broadcast(Player->GetUrl());
 
 		if (PlayOnOpen)
 		{
@@ -668,7 +675,7 @@ void UMediaPlayer::HandlePlayerMediaEvent(EMediaEvent Event)
 		break;
 
 	case EMediaEvent::MediaOpenFailed:
-		OnMediaOpenFailed.Broadcast(Player.GetUrl());
+		OnMediaOpenFailed.Broadcast(Player->GetUrl());
 
 		if (!Loop && (Playlist != nullptr))
 		{
diff --git a/Engine/Source/Runtime/MediaAssets/Public/MediaPlayer.h b/Engine/Source/Runtime/MediaAssets/Public/MediaPlayer.h
--- a/Engine/Source/Runtime/MediaAssets/Public/MediaPlayer.h
+++ b/Engine/Source/Runtime/MediaAssets/Public/MediaPlayer.h
@@ -731,7 +731,7 @@ private:
 	FOnMediaEvent MediaEvent;
 
 	/** The player implementation. */
-	FMediaPlayerBase Player;
+	TSharedPtr<FMediaPlayerBase, ESPMode::ThreadSafe> Player;
 
 #if WITH_EDITOR
 	/** Whether the player was playing in PIE/SIE. */
