Copying //UE4/Dev-Sequencer to //UE4/Main (Source: //UE4/Dev-Sequencer @ 3617642)

diff --git a/Engine/Source/Runtime/Media/Private/MediaModule.cpp b/Engine/Source/Runtime/Media/Private/MediaModule.cpp
--- a/Engine/Source/Runtime/Media/Private/MediaModule.cpp
+++ b/Engine/Source/Runtime/Media/Private/MediaModule.cpp
@@ -1,9 +1,17 @@
 // Copyright 1998-2017 Epic Games, Inc. All Rights Reserved.
 
+#include "IMediaModule.h"
+
 #include "CoreMinimal.h"
+#include "HAL/RunnableThread.h"
+#include "Misc/ScopeLock.h"
 #include "Modules/ModuleManager.h"
-#include "IMediaModule.h"
+
+#include "IMediaCaptureSupport.h"
 #include "IMediaPlayerFactory.h"
+#include "IMediaTimeSource.h"
+#include "MediaClock.h"
+#include "MediaTicker.h"
 
 
 /**
@@ -16,14 +24,24 @@ public:
 
 	//~ IMediaModule interface
 
+	virtual const TArray<IMediaCaptureSupport*>& GetCaptureSupports() const override
+	{
+		return CaptureSupports;
+	}
+
+	virtual IMediaClock& GetClock() override
+	{
+		return Clock;
+	}
+
 	virtual const TArray<IMediaPlayerFactory*>& GetPlayerFactories() const override
 	{
-		return MediaPlayerFactories;
+		return PlayerFactories;
 	}
 
 	virtual IMediaPlayerFactory* GetPlayerFactory(const FName& FactoryName) const override
 	{
-		for (IMediaPlayerFactory* Factory : MediaPlayerFactories)
+		for (IMediaPlayerFactory* Factory : PlayerFactories)
 		{
 			if (Factory->GetName() == FactoryName)
 			{
@@ -34,25 +52,83 @@ public:
 		return nullptr;
 	}
 
+	virtual IMediaTicker& GetTicker() override
+	{
+		return Ticker;
+	}
+
+	virtual void LockToTimecode(bool Locked) override
+	{
+		TimecodeLocked = Locked;
+	}
+
+	virtual void RegisterCaptureSupport(IMediaCaptureSupport& Support) override
+	{
+		CaptureSupports.AddUnique(&Support);
+	}
+
 	virtual void RegisterPlayerFactory(IMediaPlayerFactory& Factory) override
 	{
-		MediaPlayerFactories.AddUnique(&Factory);
+		PlayerFactories.AddUnique(&Factory);
+	}
+
+	virtual void SetTimeSource(const TSharedPtr<IMediaTimeSource, ESPMode::ThreadSafe>& NewTimeSource) override
+	{
+		TimeSource = NewTimeSource;
+	}
+
+	virtual void TickPostEngine() override
+	{
+		Clock.TickFetch();
+	}
+
+	virtual void TickPostRender() override
+	{
+		Clock.TickOutput();
+	}
+
+	virtual void TickPreEngine() override
+	{
+		if (TimeSource.IsValid())
+		{
+			Clock.UpdateTimecode(TimeSource->GetTimecode(), TimecodeLocked);
+		}
+
+		Clock.TickInput();
+	}
+
+	virtual void TickPreSlate() override
+	{
+		Clock.TickRender();
+	}
+
+	virtual void UnregisterCaptureSupport(IMediaCaptureSupport& Support) override
+	{
+		CaptureSupports.Remove(&Support);
 	}
 
 	virtual void UnregisterPlayerFactory(IMediaPlayerFactory& Factory) override
 	{
-		MediaPlayerFactories.Remove(&Factory);
+		PlayerFactories.Remove(&Factory);
 	}
 
 public:
 
 	//~ IModuleInterface interface
 
-	virtual void StartupModule() override { }
+	virtual void StartupModule() override
+	{
+		TickerThread = FRunnableThread::Create(&Ticker, TEXT("FMediaTicker"));
+	}
 
 	virtual void ShutdownModule() override
 	{
-		MediaPlayerFactories.Reset();
+		TickerThread->Kill(true);
+		delete TickerThread;
+		TickerThread = nullptr;
+
+		CaptureSupports.Reset();
+		PlayerFactories.Reset();
 	}
 
 	virtual bool SupportsDynamicReloading() override
@@ -62,8 +138,29 @@ public:
 
 private:
 
-	/** Holds the registered video player factories. */
-	TArray<IMediaPlayerFactory*> MediaPlayerFactories;
+	/** The registered capture device support objects. */
+	TArray<IMediaCaptureSupport*> CaptureSupports;
+
+	/** The media clock. */
+	FMediaClock Clock;
+
+	/** Time code of the current frame. */
+	FTimespan CurrentTimecode;
+
+	/** The registered video player factories. */
+	TArray<IMediaPlayerFactory*> PlayerFactories;
+
+	/** High-frequency ticker runnable. */
+	FMediaTicker Ticker;
+
+	/** High-frequency ticker thread. */
+	FRunnableThread* TickerThread;
+
+	/** Whether media objects should lock to the media clock's time code. */
+	bool TimecodeLocked;
+
+	/** The media clock's time source. */
+	TSharedPtr<IMediaTimeSource, ESPMode::ThreadSafe> TimeSource;
 };
 
 
diff --git a/Engine/Source/Runtime/Media/Public/IMediaControls.h b/Engine/Source/Runtime/Media/Public/IMediaControls.h
--- a/Engine/Source/Runtime/Media/Public/IMediaControls.h
+++ b/Engine/Source/Runtime/Media/Public/IMediaControls.h
@@ -2,25 +2,33 @@
 
 #pragma once
 
-#include "CoreMinimal.h"
+#include "Math/Range.h"
+#include "Math/RangeSet.h"
+#include "Misc/EnumClassFlags.h"
+#include "Misc/Timespan.h"
 
-class Error;
 
 /**
- * Enumerates directions for playback rate information.
+ * Available media controls.
  */
-enum class EMediaPlaybackDirections
+enum class EMediaControl
 {
-	/** Forward playback rates. */
-	Forward,
+	/** Pause playback. */
+	Pause,
 
-	/** Reverse playback rates. */
-	Reverse
+	/** Resume playback. */
+	Resume,
+
+	/** Seek to playback position (while updating output). */
+	Scrub,
+
+	/** Seek to playback position. */
+	Seek
 };
 
 
 /**
- * Enumerates directions for seeking in media.
+ * Directions for seeking in media.
  */
 enum class EMediaSeekDirection
 {
@@ -39,7 +47,20 @@ enum class EMediaSeekDirection
 
 
 /**
- * Enumerates possible states of media playback.
+ * Thinning modes for playback rates.
+ */
+enum class EMediaRateThinning
+{
+	/** Frames will be skipped to accommodate play rate. */
+	Thinned,
+
+	/** No frames will be skipped. */
+	Unthinned
+};
+
+
+/**
+ * Possible states of media playback.
  */
 enum class EMediaState
 {
@@ -55,7 +76,7 @@ enum class EMediaState
 	/** Media is currently playing. */
 	Playing,
 
-	/** Media is loading or buffering. */
+	/** Media is being prepared for playback. */
 	Preparing,
 
 	/** Playback has been stopped, but can be restarted. */
@@ -63,15 +84,41 @@ enum class EMediaState
 };
 
 
+/**
+ * Available media player status flags.
+ */
+enum class EMediaStatus
+{
+	/** No flags set. */
+	None = 0x0,
+
+	/** Player is buffering data. */
+	Buffering = 0x1,
+
+	/** Player is connecting to a media source. */
+	Connecting = 0x2
+};
+
+
+ENUM_CLASS_FLAGS(EMediaStatus);
+
+
 /**
  * Interface for controlling media playback.
  *
- * @see IMediaOutput, IMediaTracks
+ * @see IMediaCache, IMediaPlayer, IMediaSamples, IMediaTracks, IMediaView
  */
 class IMediaControls
 {
 public:
 
+	/**
+	 * Whether the specified control is currently available.
+	 *
+	 * @return true if the control is available, false otherwise.
+	 */
+	virtual bool CanControl(EMediaControl Control) const = 0;
+
 	/**
 	 * Get the media's duration.
 	 *
@@ -92,18 +139,25 @@ public:
 	 * Get the state of the media.
 	 *
 	 * @return Media state.
+	 * @see IsLooping
 	 */
 	virtual EMediaState GetState() const = 0;
 
 	/**
-	 * Get the range of supported playback rates in the specified playback direction.
+	 * Get media player status flags.
+	 *
+	 * @return Status flags.
+	 */
+	virtual EMediaStatus GetStatus() const = 0;
+
+	/**
+	 * Get the supported playback rates.
 	 *
-	 * @param Direction The playback direction.
-	 * @param Unthinned Whether the rates are for unthinned playback.
-	 * @return The range of supported rates.
-	 * @see SupportsRate, SetRate
+	 * @param Thinning The desired rate thinning mode.
+	 * @return The ranges of supported rates.
+	 * @see SetRate, SupportsRate
 	 */
-	virtual TRange<float> GetSupportedRates(EMediaPlaybackDirections Direction, bool Unthinned) const = 0;
+	virtual TRangeSet<float> GetSupportedRates(EMediaRateThinning Thinning) const = 0;
 
 	/**
 	 * Get the player's current playback time.
@@ -117,7 +171,7 @@ public:
 	 * Check whether playback is currently looping.
 	 *
 	 * @return true if playback is looping, false otherwise.
-	 * @see SetLooping
+	 * @see GetState, SetLooping
 	 */
 	virtual bool IsLooping() const = 0;
 
@@ -152,34 +206,6 @@ public:
 	 */
 	virtual bool SetRate(float Rate) = 0;
 
-	/**
-	 * Check whether the specified playback rate is supported.
-	 *
-	 * @param Rate The rate to check (can be negative for reverse play).
-	 * @param Unthinned Whether no frames should be dropped at the given rate.
-	 * @return true if the rate is supported, false otherwise.
-	 * @see GetSupportedRates, SetRate, SupportsScrubbing, SupportsSeeking
-	 */
-	virtual bool SupportsRate(float Rate, bool Unthinned) const = 0;
-
-	/**
-	 * Check whether scrubbing is supported.
-	 *
-	 * Scrubbing is the ability to decode video frames while seeking in a media item at a playback rate of 0.0.
-	 *
-	 * @return true if scrubbing is supported, false otherwise.
-	 * @see SupportsRate, SupportsSeeking
-	 */
-	virtual bool SupportsScrubbing() const = 0;
-
-	/**
-	 * Check whether the currently loaded media can jump to certain times.
-	 *
-	 * @return true if seeking is supported, false otherwise.
-	 * @see SupportsRate, SupportsScrubbing
-	 */
-	virtual bool SupportsSeeking() const = 0;
-
 public:
 
 	/**
diff --git a/Engine/Source/Runtime/Media/Public/IMediaModule.h b/Engine/Source/Runtime/Media/Public/IMediaModule.h
--- a/Engine/Source/Runtime/Media/Public/IMediaModule.h
+++ b/Engine/Source/Runtime/Media/Public/IMediaModule.h
@@ -2,24 +2,73 @@
 
 #pragma once
 
-#include "CoreMinimal.h"
+#include "Containers/Array.h"
+#include "Misc/Timespan.h"
 #include "Modules/ModuleInterface.h"
+#include "Templates/SharedPointer.h"
 
+class FName;
+class IMediaCaptureSupport;
+class IMediaClock;
 class IMediaPlayerFactory;
+class IMediaTicker;
+class IMediaTimeSource;
+
 
 /**
  * Interface for the Media module.
+ *
+ * Media Framework is ticked in several stages. The Input tick happens after the Engine
+ * (including Sequencer, game world, and Blueprints) were ticked. This allows any game
+ * logic to modify the state of tickable media objects before they fetch input.
+ *
+ * The Update tick happens right after the Input tick when all tickable media objects
+ * have fetched the latest input based on their current state.
+ *
+ * Finally, the Output tick happens after all game and core ticking completed, and
+ * after all render commands have been enqueued. It is the very last thing to happen
+ * before the frame is complete.
  */
 class IMediaModule
 	: public IModuleInterface
 {
 public:
 
+	//~ Capture devices
+
+	/**
+	 * Get all registered capture device support objects.
+	 *
+	 * @return Collection of registered objects.
+	 * @see RegisterCaptureDevices, UnregisterCaptureDevices
+	 */
+	virtual const TArray<IMediaCaptureSupport*>& GetCaptureSupports() const = 0;
+
+	/**
+	 * Register a media capture devices support object.
+	 *
+	 * @param Support The support object to register.
+	 * @see GetCaptureDevices, UnregisterCaptureDevices
+	 */
+	virtual void RegisterCaptureSupport(IMediaCaptureSupport& Support) = 0;
+
+	/**
+	 * Unregister a media capture device support object.
+	 *
+	 * @param Support The support object to unregister.
+	 * @see GetCaptureDevices, RegisterCaptureDevices
+	 */
+	virtual void UnregisterCaptureSupport(IMediaCaptureSupport& Support) = 0;
+
+public:
+
+	//~ Player factories
+
 	/**
 	 * Get the list of installed media player factories.
 	 *
 	 * @return Collection of media player factories.
-	 * @see GetPlayerInfo
+	 * @see GetPlayerFactory, RegisterPlayerFactory, UnregisterPlayerFactory
 	 */
 	virtual const TArray<IMediaPlayerFactory*>& GetPlayerFactories() const = 0;
 
@@ -28,26 +77,91 @@ public:
 	 *
 	 * @param FactoryName The name of the factory.
 	 * @return The factory, or nullptr if not found.
-	 * @see GetPlayerFactories
+	 * @see GetPlayerFactories, RegisterPlayerFactory, UnregisterPlayerFactory
 	 */
 	virtual IMediaPlayerFactory* GetPlayerFactory(const FName& FactoryName) const = 0;
 
 	/**
-	 * Registers a media player factory.
+	 * Register a media player factory.
 	 *
 	 * @param Factory The media player factory to register.
-	 * @see UnregisterPlayerFactory
+	 * @see GetPlayerFactories, RegisterPlayerFactory, UnregisterPlayerFactory
 	 */
 	virtual void RegisterPlayerFactory(IMediaPlayerFactory& Factory) = 0;
 
 	/**
-	 * Unregisters a media player factory.
+	 * Unregister a media player factory.
 	 *
 	 * @param Factory The media player factory to unregister.
-	 * @see RegisterPlayerFactory
+	 * @see GetPlayerFactories, RegisterPlayerFactory, RegisterPlayerFactory
 	 */
 	virtual void UnregisterPlayerFactory(IMediaPlayerFactory& Factory) = 0;
 
+public:
+
+	//~ Time management
+
+	/**
+	 * Get the media clock.
+	 *
+	 * @return Media clock.
+	 * @see GetTicker, LockToTimecode
+	 */
+	virtual IMediaClock& GetClock() = 0;
+
+	/**
+	 * Get the high-frequency ticker.
+	 *
+	 * @return The ticker.
+	 * @see GetClock
+	 */
+	virtual IMediaTicker& GetTicker() = 0;
+
+	/**
+	 * Whether media objects should lock to the media clock's time code.
+	 *
+	 * Time code locking changes will take effect next frame.
+	 *
+	 * @param Locked true to enable time code locking, false to disable.
+	 * @see GetClock
+	 */
+	virtual void LockToTimecode(bool Locked) = 0;
+
+	/**
+	 * Set the time source for the media clock.
+	 *
+	 * @param NewTimeSource The time source to set.
+	 */
+	virtual void SetTimeSource(const TSharedPtr<IMediaTimeSource, ESPMode::ThreadSafe>& NewTimeSource) = 0;
+
+	/**
+	 * Called by the main loop after the game engine has been ticked.
+	 *
+	 * @see TickPostRender, TickPreEngine, TickPreSlate
+	 */
+	virtual void TickPostEngine() = 0;
+
+	/**
+	 * Called by the main loop after the entire frame has been rendered.
+	 *
+	 * @see TickPostEngine, TickPreEngine, TickPreSlate
+	 */
+	virtual void TickPostRender() = 0;
+
+	/**
+	 * Called by the main loop before the game engine is ticked.
+	 *
+	 * @see TickPostEngine, TickPostRender, TickPreSlate
+	 */
+	virtual void TickPreEngine() = 0;
+
+	/**
+	 * Called by the main loop before Slate is ticked.
+	 *
+	 * @see TickPostEngine, TickPostRender, TickPreEngine
+	 */
+	virtual void TickPreSlate() = 0;
+
 public:
 
 	/** Virtual destructor. */
diff --git a/Engine/Source/Runtime/Media/Public/IMediaOptions.h b/Engine/Source/Runtime/Media/Public/IMediaOptions.h
--- a/Engine/Source/Runtime/Media/Public/IMediaOptions.h
+++ b/Engine/Source/Runtime/Media/Public/IMediaOptions.h
@@ -2,7 +2,10 @@
 
 #pragma once
 
-#include "CoreMinimal.h"
+#include "CoreTypes.h"
+#include "Containers/UnrealString.h"
+#include "Internationalization/Text.h"
+
 
 /**
  * Interface for media options.
diff --git a/Engine/Source/Runtime/Media/Public/IMediaPlayer.h b/Engine/Source/Runtime/Media/Public/IMediaPlayer.h
--- a/Engine/Source/Runtime/Media/Public/IMediaPlayer.h
+++ b/Engine/Source/Runtime/Media/Public/IMediaPlayer.h
@@ -2,52 +2,35 @@
 
 #pragma once
 
-#include "CoreMinimal.h"
+#include "Containers/UnrealString.h"
+#include "Internationalization/Text.h"
+#include "Misc/Timespan.h"
 
+class FArchive;
+
+class IMediaCache;
 class IMediaControls;
 class IMediaOptions;
-class IMediaOutput;
+class IMediaSamples;
 class IMediaTracks;
+class IMediaView;
 
-/**
- * Enumerates media player related events.
- */
-enum class EMediaEvent
-{
-	/** The current media source has been closed. */
-	MediaClosed,
-
-	/** A new media source has been opened. */
-	MediaOpened,
-
-	/** A media source failed to open. */
-	MediaOpenFailed,
-
-	/** The end of the media (or beginning if playing in reverse) has been reached. */
-	PlaybackEndReached,
-
-	/** Playback has been resumed. */
-	PlaybackResumed,
-
-	/** Playback has been suspended. */
-	PlaybackSuspended,
-
-	/** Media tracks have changed. */
-	TracksChanged
-};
+struct FGuid;
 
 
 /**
  * Interface for media players.
  *
- * @see IMediaStream
+ * @see IMediaPlayerFactory
  */
 class IMediaPlayer
 {
 public:
 
+	//~ The following methods must be implemented by media players
+
 	/**
-	 * Close a previously opened media.
+	 * Close a previously opened media source.
 	 *
 	 * Call this method to free up all resources associated with an opened
 	 * media source. If no media is open, this function has no effect.
@@ -62,10 +45,22 @@ public:
 	virtual void Close() = 0;
 
 	/**
-	 * Get the media playback controls for this player.
+	 * Get the player's cache controls.
+	 *
+	 * The interface returned by this method must remain valid for the player's life time.
+	 *
+	 * @return Cache controls.
+	 * @see GetControls, GetSamples, GetTracks, GetView
+	 */
+	virtual IMediaCache& GetCache() = 0;
+
+	/**
+	 * Get the player's playback controls.
+	 *
+	 * The interface returned by this method must remain valid for the player's life time.
 	 *
 	 * @return Playback controls.
-	 * @see GetOutput, GetTracks
+	 * @see GetCache, GetSamples, GetTracks, GetView
 	 */
 	virtual IMediaControls& GetControls() = 0;
 
@@ -85,12 +80,14 @@ public:
 	virtual FName GetName() const = 0;
 
 	/**
-	 * Get access to the media player's output.
+	 * Get the player's sample queue.
+	 *
+	 * The interface returned by this method must remain valid for the player's life time.
 	 *
-	 * @return Media tracks interface.
-	 * @see GetControls, GetTracks
+	 * @return Cache interface.
+	 * @see GetCache, GetControls, GetTracks, GetView
 	 */
-	virtual IMediaOutput& GetOutput() = 0;
+	virtual IMediaSamples& GetSamples() = 0;
 
 	/**
 	 * Get playback statistics information.
@@ -101,10 +98,12 @@ public:
 	virtual FString GetStats() const = 0;
 
 	/**
-	 * Get access to the media player's tracks.
+	 * Get the player's track collection.
 	 *
-	 * @return Media tracks interface.
-	 * @see GetControls, GetOutput
+	 * The interface returned by this method must remain valid for the player's life time.
+	 *
+	 * @return Tracks interface.
+	 * @see GetCache, GetControls, GetSamples, GetView
 	 */
 	virtual IMediaTracks& GetTracks() = 0;
 
@@ -115,6 +114,16 @@ public:
 	 */
 	virtual FString GetUrl() const = 0;
 
+	/**
+	 * Get the player's view settings.
+	 *
+	 * The interface returned by this method must remain valid for the player's life time.
+	 *
+	 * @return View interface.
+	 * @see GetCache, GetControls, GetSamples, GetTracks
+	 */
+	virtual IMediaView& GetView() = 0;
+
 	/**
 	 * Open a media source from a URL with optional parameters.
 	 *
@@ -132,7 +141,7 @@ public:
 	 * @return true if the media is being opened, false otherwise.
 	 * @see Close, IsReady, OnOpen, OnOpenFailed
 	 */
-	virtual bool Open(const FString& Url, const IMediaOptions& Options) = 0;
+	virtual bool Open(const FString& Url, const IMediaOptions* Options) = 0;
 
 	/**
 	 * Open a media source from a file or memory archive with optional parameters.
@@ -152,29 +161,62 @@ public:
 	 * @return true if the media is being opened, false otherwise.
 	 * @see Close, IsReady, OnOpen, OnOpenFailed
 	 */
-	virtual bool Open(const TSharedRef<FArchive, ESPMode::ThreadSafe>& Archive, const FString& OriginalUrl, const IMediaOptions& Options) = 0;
+	virtual bool Open(const TSharedRef<FArchive, ESPMode::ThreadSafe>& Archive, const FString& OriginalUrl, const IMediaOptions* Options) = 0;
+
+public:
+
+	//~ The following methods are optional
 
 	/**
-	 * Tick the media player logic.
+	 * Set the player's globally unique identifier.
 	 *
-	 * @param DeltaTime Time since last tick.
-	 * @see TickVideo
+	 * @param Guid The GUID to set.
 	 */
-	virtual void TickPlayer(float DeltaTime) = 0;
+	virtual void SetGuid(const FGuid& Guid)
+	{
+		// override in child classes if supported
+	}
 
 	/**
-	 * Tick the media player's video code.
+	 * Tick the player's audio related code.
 	 *
-	 * @param DeltaTime Time since last tick.
-	 * @see TickPlayer
+	 * This is a high-frequency tick function. Media players override this method
+	 * to fetch and process audio samples, or to perform other time critical tasks.
+	 *
+	 * @see TickInput, TickFetch
 	 */
-	virtual void TickVideo(float DeltaTime) = 0;
+	virtual void TickAudio()
+	{
+		// override in child class if needed
+	}
 
-public:
+	/**
+	 * Tick the player in the Fetch phase.
+	 *
+	 * Media players may override this method to fetch newly decoded input
+	 * samples before they are rendered on textures or audio components.
+	 *
+	 * @param DeltaTime Time since last tick.
+	 * @see TickAudio, TickInput
+	 */
+	virtual void TickFetch(FTimespan DeltaTime)
+	{
+		// override in child class if needed
+	}
 
-	/** Get an event delegate that is invoked when an event occurred. */
-	DECLARE_EVENT_OneParam(IMediaPlayer, FOnMediaEvent, EMediaEvent /*Event*/)
-	virtual FOnMediaEvent& OnMediaEvent() = 0;
+	/**
+	 * Tick the player in the Input phase.
+	 *
+	 * Media players may override this method to update their state before the
+	 * Engine is being ticked, or to initiate the processing of input samples.
+	 *
+	 * @param DeltaTime Time since last tick.
+	 * @see TickAudio, TickFetch
+	 */
+	virtual void TickInput(FTimespan DeltaTime)
+	{
+		// override in child class if needed
+	}
 
 public:
 
diff --git a/Engine/Source/Runtime/Media/Public/IMediaPlayerFactory.h b/Engine/Source/Runtime/Media/Public/IMediaPlayerFactory.h
--- a/Engine/Source/Runtime/Media/Public/IMediaPlayerFactory.h
+++ b/Engine/Source/Runtime/Media/Public/IMediaPlayerFactory.h
@@ -2,11 +2,57 @@
 
 #pragma once
 
-#include "CoreMinimal.h"
+#include "CoreTypes.h"
+#include "Containers/Array.h"
+#include "Internationalization/Text.h"
+#include "Templates/SharedPointer.h"
 
+class FString;
+class IMediaEventSink;
 class IMediaOptions;
 class IMediaPlayer;
 
+
+/**
+ * Enumerates available media player features.
+ */
+enum class EMediaFeature
+{
+	/** Audio output via Engine. */
+	AudioSamples,
+
+	/** Audio tracks. */
+	AudioTracks,
+
+	/** Caption tracks. */
+	CaptionTracks,
+
+	/** Metadata tracks (implies output via Engine). */
+	MetadataTracks,
+
+	/** Captions, subtitle and text output via Engine. */
+	OverlaySamples,
+
+	/** Subtitle tracks. */
+	SubtitleTracks,
+
+	/** Generic text tracks. */
+	TextTracks,
+
+	/** 360 degree video controls. */
+	Video360,
+
+	/** Video output via Engine. */
+	VideoSamples,
+
+	/** Stereoscopic video controls. */
+	VideoStereo,
+
+	/** Video tracks. */
+	VideoTracks
+};
+
+
 /**
  * Interface for media player factories.
  *
@@ -14,6 +60,8 @@ class IMediaPlayer;
  * Most media players will be implemented inside plug-ins, which will register their
  * factories on startup. The Media module will use the CanPlayUrl() method on this
  * interface to determine which media player to instantiate for a given media source.
+ *
+ * @see IMediaPlayer
  */
 class IMediaPlayerFactory
 {
@@ -28,14 +76,15 @@ public:
 	 * @param OutErrors will contain error messages (optional).
 	 * @return true if the source can be played, false otherwise.
 	 */
-	virtual bool CanPlayUrl(const FString& Url, const IMediaOptions& Options, TArray<FText>* OutWarnings, TArray<FText>* OutErrors) const = 0;
+	virtual bool CanPlayUrl(const FString& Url, const IMediaOptions* Options, TArray<FText>* OutWarnings, TArray<FText>* OutErrors) const = 0;
 
 	/**
 	 * Creates a media player.
 	 *
+	 * @param EventSink The object that will receive the player's events.
 	 * @return A new media player, or nullptr if a player couldn't be created.
 	 */
-	virtual TSharedPtr<IMediaPlayer, ESPMode::ThreadSafe> CreatePlayer() = 0;
+	virtual TSharedPtr<IMediaPlayer, ESPMode::ThreadSafe> CreatePlayer(IMediaEventSink& EventSink) = 0;
 
 	/**
 	 * Get the human readable name of the player.
@@ -64,6 +113,14 @@ public:
 	 */
 	virtual const TArray<FString>& GetSupportedPlatforms() const = 0;
 
+	/**
+	 * Check whether the media player supports the specified feature.
+	 *
+	 * @param Feature The feature to check.
+	 * @return true if the feature is supported, false otherwise.
+	 */
+	virtual bool SupportsFeature(EMediaFeature Feature) const = 0;
+
 public:
 
 	/**
@@ -73,7 +130,7 @@ public:
 	 * @param Options Optional media player parameters.
 	 * @return true if the source can be played, false otherwise.
 	 */
-	bool CanPlayUrl(const FString& Url, const IMediaOptions& Options) const
+	bool CanPlayUrl(const FString& Url, const IMediaOptions* Options) const
 	{
 		return CanPlayUrl(Url, Options, nullptr, nullptr);
 	}
diff --git a/Engine/Source/Runtime/Media/Public/IMediaTickable.h b/Engine/Source/Runtime/Media/Public/IMediaTickable.h
new file mode 100644
--- /dev/null
+++ b/Engine/Source/Runtime/Media/Public/IMediaTickable.h
@@ -0,0 +1,20 @@
+// Copyright 1998-2017 Epic Games, Inc. All Rights Reserved.
+
+#pragma once
+
+
+/**
+ * Interface for objects that can be ticked from the high-frequency media thread.
+ */
+class IMediaTickable
+{
+public:
+
+	/** Tick the object. */
+	virtual void TickTickable() = 0;
+
+public:
+
+	/** Virtual destructor. */
+	virtual ~IMediaTickable() { }
+};
diff --git a/Engine/Source/Runtime/Media/Public/IMediaTimeSource.h b/Engine/Source/Runtime/Media/Public/IMediaTimeSource.h
new file mode 100644
--- /dev/null
+++ b/Engine/Source/Runtime/Media/Public/IMediaTimeSource.h
@@ -0,0 +1,24 @@
+// Copyright 1998-2017 Epic Games, Inc. All Rights Reserved.
+
+#pragma once
+
+
+/**
+ * Interface for media clock time sources.
+ */
+class IMediaTimeSource
+{
+public:
+
+	/**
+	 * Get the current time code.
+	 *
+	 * @return Time code.
+	 */
+	virtual FTimespan GetTimecode() = 0;
+
+public:
+
+	/** Virtual destructor. */
+	virtual ~IMediaTimeSource() { }
+};
diff --git a/Engine/Source/Runtime/Media/Public/IMediaTracks.h b/Engine/Source/Runtime/Media/Public/IMediaTracks.h
--- a/Engine/Source/Runtime/Media/Public/IMediaTracks.h
+++ b/Engine/Source/Runtime/Media/Public/IMediaTracks.h
@@ -2,10 +2,17 @@
 
 #pragma once
 
-#include "CoreMinimal.h"
+#include "CoreTypes.h"
+#include "Internationalization/Text.h"
+#include "Math/IntPoint.h"
+#include "Math/Range.h"
+#include "Math/RangeSet.h"
+
 
 /**
  * Enumerates available media track types.
+ *
+ * Note: Keep this in sync with EMediaTrackType
  */
 enum class EMediaTrackType
 {
@@ -32,42 +39,83 @@ enum class EMediaTrackType
 };
 
 
+/**
+ * Audio track format details.
+ */
+struct FMediaAudioTrackFormat
+{
+	/** Number of bits per sample. */
+	uint32 BitsPerSample;
+
+	/** Number of audio channels. */
+	uint32 NumChannels;
+
+	/** Sample rate (in samples per second). */
+	uint32 SampleRate;
+
+	/** Name of the format type. */
+	FString TypeName;
+};
+
+
+/**
+ * Video track format details.
+ */
+struct FMediaVideoTrackFormat
+{
+	/** Width and height of the video (in pixels). */
+	FIntPoint Dim;
+
+	/** Active frame rate (in frames per second). */
+	float FrameRate;
+
+	/** Supported frame rate range. */
+	TRange<float> FrameRates;
+
+	/** Name of the format type. */
+	FString TypeName;
+};
+
+
 /**
  * Interface for access to a media player's tracks.
  *
- * @see IMediaControls, IMediaOutput
+ * @see IMediaCache, IMediaControls, IMediaPlayer, IMediaSamples, IMediaView
  */
 class IMediaTracks
 {
 public:
 
 	/**
-	 * Get the number of channels in the specified audio track.
+	 * Get details about the specified audio track format.
 	 *
-	 * @param TrackIndex Index of the audio track.
-	 * @return Number of channels.
-	 * @see GetAudioTrackSamplesPerSeconds
+	 * @param TrackIndex The index of the audio track (must be valid).
+	 * @param FormatIndex The index of the track's format (must be valid).
+	 * @param OutFormat Will contain the format details.
+	 * @return true on success, false otherwise.
+	 * @see GetVideoTrackFormat
 	 */
-	virtual uint32 GetAudioTrackChannels(int32 TrackIndex) const = 0;
-
-	/**
-	 * Get the sample rate of the specified audio track.
-	 *
-	 * @param TrackIndex Index of the audio track.
-	 * @return Samples per second.
-	 * @see GetAudioTrackChannels
-	 */
-	virtual uint32 GetAudioTrackSampleRate(int32 TrackIndex) const = 0;
+	virtual bool GetAudioTrackFormat(int32 TrackIndex, int32 FormatIndex, FMediaAudioTrackFormat& OutFormat) const = 0;
 
 	/**
 	 * Get the number of media tracks of the given type.
 	 *
 	 * @param TrackType The type of media tracks.
 	 * @return Number of tracks.
-	 * @see GetSelectedTrack, SelectTrack
+	 * @see GetNumTrackFormats, GetSelectedTrack, SelectTrack
 	 */
 	virtual int32 GetNumTracks(EMediaTrackType TrackType) const = 0;
 
+	/**
+	 * Get the number of formats of the specified track.
+	 *
+	 * @param TrackType The type of track.
+	 * @param TrackIndex The index of the track (must be valid).
+	 * @return Number of formats.
+	 * @see GetNumTracks, GetTrackFormatType, GetTrackDisplayName
+	 */
+	virtual int32 GetNumTrackFormats(EMediaTrackType TrackType, int32 TrackIndex) const = 0;
+
 	/**
 	 * Get the index of the currently selected track of the given type.
 	 *
@@ -81,17 +129,27 @@ public:
 	 * Get the human readable name of the specified track.
 	 *
 	 * @param TrackType The type of track.
-	 * @param TrackIndex The index of the track.
+	 * @param TrackIndex The index of the track (must be valid).
 	 * @return Display name.
 	 * @see GetTrackLanguage, GetTrackName
 	 */
 	virtual FText GetTrackDisplayName(EMediaTrackType TrackType, int32 TrackIndex) const = 0;
 
+	/**
+	 * Get the index of the currently selected format of the given track type.
+	 *
+	 * @param TrackType The type of track.
+	 * @param TrackIndex The index of the track (must be valid).
+	 * @return The index of the selected format.
+	 * @see GetNumTrackFormats, GetSelectedTrack, SetTrackFormat
+	 */
+	virtual int32 GetTrackFormat(EMediaTrackType TrackType, int32 TrackIndex) const = 0;
+
 	/**
 	 * Get the language tag of the specified track.
 	 *
 	 * @param TrackType The type of track.
-	 * @param TrackIndex The index of the track.
+	 * @param TrackIndex The index of the track (must be valid).
 	 * @return Language tag, i.e. "en-US" for English, or "und" for undefined.
 	 * @see GetTrackDisplayName, GetTrackName
 	 */
@@ -101,72 +159,67 @@ public:
 	 * Get the internal name of the specified track.
 	 *
 	 * @param TrackType The type of track.
-	 * @param TrackIndex The index of the track.
+	 * @param TrackIndex The index of the track (must be valid).
 	 * @return Track name.
 	 * @see GetTrackDisplayName, GetTrackLanguage
 	 */
 	virtual FString GetTrackName(EMediaTrackType TrackType, int32 TrackIndex) const = 0;
 
 	/**
-	 * Get the average data rate of the specified video track.
-	 *
-	 * @param TrackIndex The index of the track.
-	 * @return Data rate (in bits per second).
-	 * @see GetVideoTrackAspectRatio, GetVideoTrackDimensions, GetVideoTrackFrameRate
-	 */
-	virtual uint32 GetVideoTrackBitRate(int32 TrackIndex) const = 0;
-
-	/**
-	 * Get the dimensions of the specified video track.
-	 *
-	 * @param TrackType The type of track.
-	 * @param TrackIndex The index of the track.
-	 * @return Video dimensions (in pixels).
-	 * @see GetVideoTrackAspectRatio, GetVideoTrackBitRate, GetVideoTrackFrameRate
-	 */
-	virtual FIntPoint GetVideoTrackDimensions(int32 TrackIndex) const = 0;
-
-	/**
-	 * Get the frame rate of the specified video track.
+	 * Get details about the specified video track format.
 	 *
-	 * @param TrackType The type of track.
-	 * @param TrackIndex The index of the track.
-	 * @return Frame rate (in frames per second).
-	 * @see GetVideoTrackAspectRatio, GetVideoTrackBitRate, GetVideoTrackDimensions
+	 * @param TrackIndex The index of the video track (must be valid).
+	 * @param FormatIndex The index of the track's format (must be valid).
+	 * @param OutFormat Will contain the format details.
+	 * @return true on success, false otherwise.
+	 * @see GetAudioTrackFormat
 	 */
-	virtual float GetVideoTrackFrameRate(int32 TrackIndex) const = 0;
+	virtual bool GetVideoTrackFormat(int32 TrackIndex, int32 FormatIndex, FMediaVideoTrackFormat& OutFormat) const = 0;
 
 	/**
 	 * Select the active track of the given type.
 	 *
-	 * Only one track of a given type can be active at any time.
+	 * The selected track will use its currently active format. Active formats will
+	 * be remembered on a per track basis. The first available format is active by
+	 * default. To switch the track format, use SetTrackFormat instead.
 	 *
 	 * @param TrackType The type of track to select.
 	 * @param TrackIndex The index of the track to select, or INDEX_NONE to deselect.
 	 * @return true if the track was selected, false otherwise.
-	 * @see GetNumTracks, GetSelectedTrack
+	 * @see GetNumTracks, GetSelectedTrack, SetTrackFormat
 	 */
 	virtual bool SelectTrack(EMediaTrackType TrackType, int32 TrackIndex) = 0;
 
-public:
+	/**
+	 * Select the active format on the specified track.
+	 *
+	 * Selecting the format will not switch to the specified track. To switch
+	 * tracks, use SelectTrack instead. If the track is already selected, the
+	 * format change will be applied immediately.
+	 *
+	 * @param TrackType The type of track to update.
+	 * @param TrackIndex The index of the track to update (must be valid).
+	 * @param FormatIndex The index of the format to select (must be valid).
+	 * @return true if the format was selected, false otherwise.
+	 * @see GetNumTrackFormats, GetNumTracks, GetTrackFormat, SelectTrack
+	 */
+	virtual bool SetTrackFormat(EMediaTrackType TrackType, int32 TrackIndex, int32 FormatIndex) = 0;
 
 	/**
-	 * Get the video's aspect ratio of the specified video track.
+	 * Set the frame rate of the specified video track.
 	 *
-	 * @param TrackIndex Index of the video track.
-	 * @return Aspect ratio.
-	 * @see GetVideoTrackBitRate, GetVideoTrackDimensions, GetVideoTrackFrameRate
+	 * Note that most players may not support overriding the video frame rate.
+	 * This feature is often only available on video capture media sources. 
+	 *
+	 * @param TrackIndex The index of the track (must be valid).
+	 * @param FormatIndex Index of the track format (must be valid).
+	 * @param FrameRate The frame rate to set (must be in range of format's supported frame rates).
+	 * @return true on success, false otherwise.
+	 * @see GetTrackFormat, SetTrackFormat
 	 */
-	float GetVideoTrackAspectRatio(int32 TrackIndex) const
+	virtual bool SetVideoTrackFrameRate(int32 TrackIndex, int32 FormatIndex, float FrameRate)
 	{
-		const FIntPoint Dimensions = GetVideoTrackDimensions(TrackIndex);
-
-		if (Dimensions.Y == 0)
-		{
-			return 0.0f;
-		}
-
-		return (float)(Dimensions.X) / Dimensions.Y;
+		return false;
 	}
 
 public:
diff --git a/Engine/Source/Runtime/MediaAssets/MediaAssets.Build.cs b/Engine/Source/Runtime/MediaAssets/MediaAssets.Build.cs
--- a/Engine/Source/Runtime/MediaAssets/MediaAssets.Build.cs
+++ b/Engine/Source/Runtime/MediaAssets/MediaAssets.Build.cs
@@ -6,32 +6,30 @@ namespace UnrealBuildTool.Rules
 	{
 		public MediaAssets(ReadOnlyTargetRules Target) : base(Target)
 		{
-            PublicDependencyModuleNames.AddRange(
-                new string[] {
+			PublicDependencyModuleNames.AddRange(
+				new string[] {
+					"AudioMixer",
 					"Core",
 					"CoreUObject",
 					"Engine",
-                    "Media",
+					"Media",
 					"MediaUtils",
-				}
-            );
+				});
 
 			PrivateDependencyModuleNames.AddRange(
 				new string[] {
-                    "RenderCore",
-                    "RHI",
+					"RenderCore",
+					"RHI",
 					"ShaderCore",
 					"UtilityShaders",
-				}
-			);
+				});
 
 			PrivateIncludePaths.AddRange(
 				new string[] {
 					"Runtime/MediaAssets/Private",
 					"Runtime/MediaAssets/Private/Assets",
 					"Runtime/MediaAssets/Private/Misc",
-				}
-			);
+				});
 
 			if (Target.bBuildEditor)
 			{
diff --git a/Engine/Source/Runtime/MediaAssets/Private/Assets/FileMediaSource.cpp b/Engine/Source/Runtime/MediaAssets/Private/Assets/FileMediaSource.cpp
--- a/Engine/Source/Runtime/MediaAssets/Private/Assets/FileMediaSource.cpp
+++ b/Engine/Source/Runtime/MediaAssets/Private/Assets/FileMediaSource.cpp
@@ -14,6 +14,22 @@ namespace FileMediaSource
 /* UFileMediaSource interface
  *****************************************************************************/
 
+FString UFileMediaSource::GetFullPath() const
+{
+	if (!FPaths::IsRelative(FilePath))
+	{
+		return FilePath;
+	}
+
+	if (FilePath.StartsWith(TEXT("./")))
+	{
+		return FPaths::ConvertRelativePathToFull(FPaths::ProjectContentDir(), FilePath.RightChop(2));
+	}
+
+	return FPaths::ConvertRelativePathToFull(FilePath);
+}
+
+
 void UFileMediaSource::SetFilePath(const FString& Path)
 {
 	if (Path.IsEmpty() || Path.StartsWith(TEXT("./")))
@@ -74,22 +90,3 @@ bool UFileMediaSource::Validate() const
 {
 	return FPaths::FileExists(GetFullPath());
 }
-
-
-/* UFileMediaSource implementation
- *****************************************************************************/
-
-FString UFileMediaSource::GetFullPath() const
-{
-	if (!FPaths::IsRelative(FilePath))
-	{
-		return FilePath;
-	}
-
-	if (FilePath.StartsWith(TEXT("./")))
-	{
-		return FPaths::ConvertRelativePathToFull(FPaths::ProjectContentDir(), FilePath.RightChop(2));
-	}
-
-	return FPaths::ConvertRelativePathToFull(FilePath);
-}
\ No newline at end of file
diff --git a/Engine/Source/Runtime/MediaAssets/Private/Assets/MediaPlayer.cpp b/Engine/Source/Runtime/MediaAssets/Private/Assets/MediaPlayer.cpp
--- a/Engine/Source/Runtime/MediaAssets/Private/Assets/MediaPlayer.cpp
+++ b/Engine/Source/Runtime/MediaAssets/Private/Assets/MediaPlayer.cpp
@@ -1,21 +1,26 @@
 // Copyright 1998-2017 Epic Games, Inc. All Rights Reserved.
 
 #include "MediaPlayer.h"
+#include "MediaAssetsPrivate.h"
 
+#include "IMediaClock.h"
 #include "IMediaControls.h"
 #include "IMediaModule.h"
-#include "IMediaOutput.h"
 #include "IMediaPlayer.h"
 #include "IMediaPlayerFactory.h"
+#include "IMediaTicker.h"
 #include "IMediaTracks.h"
-#include "MediaOverlays.h"
+#include "MediaPlayerFacade.h"
+#include "Misc/App.h"
 #include "Misc/Paths.h"
-#include "MediaAssetsPrivate.h"
+#include "Misc/ScopeLock.h"
+#include "Modules/ModuleManager.h"
+#include "UObject/Package.h"
+#include "UObject/UObjectGlobals.h"
+
 #include "MediaPlaylist.h"
 #include "MediaSource.h"
-#include "MediaSoundWave.h"
-#include "MediaTexture.h"
-#include "Modules/ModuleManager.h"
+#include "StreamMediaSource.h"
 
 
 /* UMediaPlayer structors
@@ -23,16 +28,25 @@
 
 UMediaPlayer::UMediaPlayer(const FObjectInitializer& ObjectInitializer)
 	: Super(ObjectInitializer)
+	, CacheAhead(FTimespan::FromMilliseconds(100))
+	, CacheBehind(FTimespan::FromMilliseconds(3000))
+	, CacheBehindGame(FTimespan::FromMilliseconds(100))
 	, PlayOnOpen(true)
 	, Shuffle(false)
 	, Loop(false)
 	, PlaylistIndex(INDEX_NONE)
-	, Player(MakeShareable(new FMediaPlayerBase))
+	, HorizontalFieldOfView(90.0f)
+	, VerticalFieldOfView(60.0f)
+	, ViewRotation(FRotator::ZeroRotator)
+	, PlayerFacade(MakeShareable(new FMediaPlayerFacade))
+	, PlayerGuid(FGuid::NewGuid())
+	, PlayOnNext(false)
 #if WITH_EDITOR
 	, WasPlayingInPIE(false)
 #endif
 {
-	Player->OnMediaEvent().AddUObject(this, &UMediaPlayer::HandlePlayerMediaEvent);
+	PlayerFacade->OnMediaEvent().AddUObject(this, &UMediaPlayer::HandlePlayerMediaEvent);
+	Playlist = NewObject<UMediaPlaylist>(GetTransientPackage(), NAME_None, RF_Transactional | RF_Transient);
 }
 
 
@@ -41,7 +55,7 @@ UMediaPlayer::UMediaPlayer(const FObjectInitializer& ObjectInitializer)
 
 bool UMediaPlayer::CanPause() const
 {
-	return Player->CanPause();
+	return PlayerFacade->CanPause();
 }
 
 
@@ -52,7 +66,7 @@ bool UMediaPlayer::CanPlaySource(UMediaSource* MediaSource)
 		return false;
 	}
 
-	return Player->CanPlayUrl(MediaSource->GetUrl(), *MediaSource);
+	return PlayerFacade->CanPlayUrl(MediaSource->GetUrl(), MediaSource);
 }
 
 
@@ -63,135 +77,244 @@ bool UMediaPlayer::CanPlayUrl(const FString& Url)
 		return false;
 	}
 
-	return Player->CanPlayUrl(Url, *GetDefault<UMediaSource>());
+	return PlayerFacade->CanPlayUrl(Url, GetDefault<UMediaSource>());
 }
 
 
 void UMediaPlayer::Close()
 {
-	Player->Close();
+	PlayerFacade->Close();
 
-	LastUrl.Empty();
-	Playlist = nullptr;
+	Playlist = NewObject<UMediaPlaylist>(GetTransientPackage(), NAME_None, RF_Transactional | RF_Transient);
 	PlaylistIndex = INDEX_NONE;
+	PlayOnNext = false;
+}
+
+
+int32 UMediaPlayer::GetAudioTrackChannels(int32 TrackIndex, int32 FormatIndex) const
+{
+	return PlayerFacade->GetAudioTrackChannels(TrackIndex, FormatIndex);
+}
+
+
+int32 UMediaPlayer::GetAudioTrackSampleRate(int32 TrackIndex, int32 FormatIndex) const
+{
+	return PlayerFacade->GetAudioTrackSampleRate(TrackIndex, FormatIndex);
 }
 
 
-FMediaPlayerBase& UMediaPlayer::GetBasePlayer()
+FString UMediaPlayer::GetAudioTrackType(int32 TrackIndex, int32 FormatIndex) const
 {
-	return *Player;
+	return PlayerFacade->GetAudioTrackType(TrackIndex, FormatIndex);
 }
 
 
 FName UMediaPlayer::GetDesiredPlayerName() const
 {
-	return Player->DesiredPlayerName;
+	return PlayerFacade->DesiredPlayerName;
 }
 
 
 FTimespan UMediaPlayer::GetDuration() const
 {
-	return Player->GetDuration();
+	return PlayerFacade->GetDuration();
 }
 
 
-FFloatRange UMediaPlayer::GetForwardRates(bool Unthinned)
+float UMediaPlayer::GetHorizontalFieldOfView() const
 {
-	return Player->GetForwardRates(Unthinned);
+	float OutHorizontal = 0.0f;
+	float OutVertical = 0.0f;
+
+	if (!PlayerFacade->GetViewField(OutHorizontal, OutVertical))
+	{
+		return 0.0f;
+	}
+
+	return OutHorizontal;
 }
 
 
 int32 UMediaPlayer::GetNumTracks(EMediaPlayerTrack TrackType) const
 {
-	return Player->GetNumTracks((EMediaTrackType)TrackType);
+	return PlayerFacade->GetNumTracks((EMediaTrackType)TrackType);
+}
+
+
+int32 UMediaPlayer::GetNumTrackFormats(EMediaPlayerTrack TrackType, int32 TrackIndex) const
+{
+	return PlayerFacade->GetNumTrackFormats((EMediaTrackType)TrackType, TrackIndex);
+}
+
+
+TSharedRef<FMediaPlayerFacade, ESPMode::ThreadSafe> UMediaPlayer::GetPlayerFacade() const
+{
+	return PlayerFacade.ToSharedRef();
 }
 
 
 FName UMediaPlayer::GetPlayerName() const
 {
-	return Player->GetPlayerName();
+	return PlayerFacade->GetPlayerName();
 }
 
 
 float UMediaPlayer::GetRate() const
 {
-	return Player->GetRate();
+	return PlayerFacade->GetRate();
 }
 
 
-FFloatRange UMediaPlayer::GetReverseRates(bool Unthinned)
+int32 UMediaPlayer::GetSelectedTrack(EMediaPlayerTrack TrackType) const
 {
-	return Player->GetReverseRates(Unthinned);
+	return PlayerFacade->GetSelectedTrack((EMediaTrackType)TrackType);
 }
 
 
-int32 UMediaPlayer::GetSelectedTrack(EMediaPlayerTrack TrackType) const
+void UMediaPlayer::GetSupportedRates(TArray<FFloatRange>& OutRates, bool Unthinned) const
 {
-	return Player->GetSelectedTrack((EMediaTrackType)TrackType);
+	const TRangeSet<float> Rates = PlayerFacade->GetSupportedRates(Unthinned);
+	Rates.GetRanges((TArray<TRange<float>>&)OutRates);
 }
 
 
 FTimespan UMediaPlayer::GetTime() const
 {
-	return Player->GetTime();
+	return PlayerFacade->GetTime();
 }
 
 
 FText UMediaPlayer::GetTrackDisplayName(EMediaPlayerTrack TrackType, int32 TrackIndex) const
 {
-	return Player->GetTrackDisplayName((EMediaTrackType)TrackType, TrackIndex);
+	return PlayerFacade->GetTrackDisplayName((EMediaTrackType)TrackType, TrackIndex);
+}
+
+
+int32 UMediaPlayer::GetTrackFormat(EMediaPlayerTrack TrackType, int32 TrackIndex) const
+{
+	return PlayerFacade->GetTrackFormat((EMediaTrackType)TrackType, TrackIndex);
 }
 
 
 FString UMediaPlayer::GetTrackLanguage(EMediaPlayerTrack TrackType, int32 TrackIndex) const
 {
-	return Player->GetTrackLanguage((EMediaTrackType)TrackType, TrackIndex);
+	return PlayerFacade->GetTrackLanguage((EMediaTrackType)TrackType, TrackIndex);
 }
 
 
 const FString& UMediaPlayer::GetUrl() const
 {
-	return Player->GetUrl();
+	return PlayerFacade->GetUrl();
+}
+
+
+float UMediaPlayer::GetVerticalFieldOfView() const
+{
+	float OutHorizontal = 0.0f;
+	float OutVertical = 0.0f;
+
+	if (!PlayerFacade->GetViewField(OutHorizontal, OutVertical))
+	{
+		return 0.0f;
+	}
+
+	return OutVertical;
+}
+
+
+float UMediaPlayer::GetVideoTrackAspectRatio(int32 TrackIndex, int32 FormatIndex) const
+{
+	return PlayerFacade->GetVideoTrackAspectRatio(TrackIndex, FormatIndex);
+}
+
+
+FIntPoint UMediaPlayer::GetVideoTrackDimensions(int32 TrackIndex, int32 FormatIndex) const
+{
+	return PlayerFacade->GetVideoTrackDimensions(TrackIndex, FormatIndex);
+}
+
+
+float UMediaPlayer::GetVideoTrackFrameRate(int32 TrackIndex, int32 FormatIndex) const
+{
+	return PlayerFacade->GetVideoTrackFrameRate(TrackIndex, FormatIndex);
+}
+
+
+FFloatRange UMediaPlayer::GetVideoTrackFrameRates(int32 TrackIndex, int32 FormatIndex) const
+{
+	return PlayerFacade->GetVideoTrackFrameRates(TrackIndex, FormatIndex);
+}
+
+
+FString UMediaPlayer::GetVideoTrackType(int32 TrackIndex, int32 FormatIndex) const
+{
+	return PlayerFacade->GetVideoTrackType(TrackIndex, FormatIndex);
+}
+
+
+FRotator UMediaPlayer::GetViewRotation() const
+{
+	FQuat OutOrientation;
+
+	if (!PlayerFacade->GetViewOrientation(OutOrientation))
+	{
+		return FRotator::ZeroRotator;
+	}
+
+	return OutOrientation.Rotator();
+}
+
+
+bool UMediaPlayer::HasError() const
+{
+	return PlayerFacade->HasError();
+}
+
+
+bool UMediaPlayer::IsBuffering() const
+{
+	return PlayerFacade->IsBuffering();
+}
+
+
+bool UMediaPlayer::IsConnecting() const
+{
+	return PlayerFacade->IsConnecting();
 }
 
 
 bool UMediaPlayer::IsLooping() const
 {
-	return Player->IsLooping();
+	return PlayerFacade->IsLooping();
 }
 
 
 bool UMediaPlayer::IsPaused() const
 {
-	return Player->IsPaused();
+	return PlayerFacade->IsPaused();
 }
 
 
 bool UMediaPlayer::IsPlaying() const
 {
-	return Player->IsPlaying();
+	return PlayerFacade->IsPlaying();
 }
 
 
 bool UMediaPlayer::IsPreparing() const
 {
-	return Player->IsPreparing();
+	return PlayerFacade->IsPreparing();
 }
 
 
 bool UMediaPlayer::IsReady() const
 {
-	return Player->IsReady();
+	return PlayerFacade->IsReady();
 }
 
 
 bool UMediaPlayer::Next()
 {
-	if (Playlist == nullptr)
-	{
-		return false;
-	}
-
 	int32 RemainingAttempts = Playlist->Num();
 
 	while (RemainingAttempts-- > 0)
@@ -200,7 +323,7 @@ bool UMediaPlayer::Next()
 			? Playlist->GetRandom(PlaylistIndex)
 			: Playlist->GetNext(PlaylistIndex);
 
-		if ((NextSource != nullptr) && NextSource->Validate() && Player->Open(NextSource->GetUrl(), *NextSource))
+		if ((NextSource != nullptr) && NextSource->Validate() && PlayerFacade->Open(NextSource->GetUrl(), NextSource))
 		{
 			return true;
 		}
@@ -212,19 +335,14 @@ bool UMediaPlayer::Next()
 
 bool UMediaPlayer::OpenFile(const FString& FilePath)
 {
-	FString FullPath;
-	
-	if (FPaths::IsRelative(FilePath))
-	{
-		FullPath = FPaths::ConvertRelativePathToFull(FilePath);
-	}
-	else
+	Close();
+
+	if (!Playlist->AddFile(FilePath))
 	{
-		FullPath = FilePath;
-		FPaths::NormalizeFilename(FullPath);
+		return false;
 	}
 
-	return OpenUrl(FString(TEXT("file://")) + FullPath);
+	return Next();
 }
 
 
@@ -252,18 +370,16 @@ bool UMediaPlayer::OpenPlaylistIndex(UMediaPlaylist* InPlaylist, int32 Index)
 		UE_LOG(LogMediaAssets, Warning, TEXT("UMediaPlayer::OpenPlaylistIndex called with invalid PlaylistIndex %i"), Index);
 		return false;
 	}
-	
-	LastUrl = MediaSource->GetUrl();
+
 	PlaylistIndex = Index;
 
 	if (!MediaSource->Validate())
 	{
 		UE_LOG(LogMediaAssets, Error, TEXT("Failed to validate media source %s (%s)"), *MediaSource->GetName(), *MediaSource->GetUrl());
-
 		return false;
 	}
 
-	return Player->Open(LastUrl, *MediaSource);
+	return PlayerFacade->Open(MediaSource->GetUrl(), MediaSource);
 }
 
 
@@ -277,16 +393,15 @@ bool UMediaPlayer::OpenSource(UMediaSource* MediaSource)
 		return false;
 	}
 
-	LastUrl = MediaSource->GetUrl();
-
 	if (!MediaSource->Validate())
 	{
 		UE_LOG(LogMediaAssets, Error, TEXT("Failed to validate media source %s (%s)"), *MediaSource->GetName(), *MediaSource->GetUrl());
-
 		return false;
 	}
 
-	return Player->Open(LastUrl, *MediaSource);
+	Playlist->Add(MediaSource);
+
+	return Next();
 }
 
 
@@ -294,45 +409,38 @@ bool UMediaPlayer::OpenUrl(const FString& Url)
 {
 	Close();
 
-	if (Url.IsEmpty())
+	if (!Playlist->AddUrl(Url))
 	{
 		return false;
 	}
 
-	LastUrl = Url;
-
-	return Player->Open(LastUrl, *GetDefault<UMediaSource>());
+	return Next();
 }
 
 
 bool UMediaPlayer::Pause()
 {
-	return Player->SetRate(0.0f);
+	return PlayerFacade->SetRate(0.0f);
 }
 
 
 bool UMediaPlayer::Play()
 {
-	return Player->SetRate(1.0f);
+	return PlayerFacade->SetRate(1.0f);
 }
 
 
 bool UMediaPlayer::Previous()
 {
-	if (Playlist == nullptr)
-	{
-		return false;
-	}
-
 	int32 RemainingAttempts = Playlist->Num();
 
 	while (--RemainingAttempts >= 0)
 	{
 		UMediaSource* PrevSource = Shuffle
 			? Playlist->GetRandom(PlaylistIndex)
-			: Playlist->GetNext(PlaylistIndex);
+			: Playlist->GetPrevious(PlaylistIndex);
 
-		if ((PrevSource != nullptr) && PrevSource->Validate() && Player->Open(PrevSource->GetUrl(), *PrevSource))
+		if ((PrevSource != nullptr) && PrevSource->Validate() && PlayerFacade->Open(PrevSource->GetUrl(), PrevSource))
 		{
 			return true;
 		}
@@ -344,12 +452,12 @@ bool UMediaPlayer::Previous()
 
 bool UMediaPlayer::Reopen()
 {
-	if (Playlist != nullptr)
+	if (Playlist == nullptr)
 	{
-		return OpenPlaylistIndex(Playlist, PlaylistIndex);
+		return false;
 	}
 
-	return Player->Open(LastUrl, *GetDefault<UMediaSource>());
+	return OpenPlaylistIndex(Playlist, PlaylistIndex);
 }
 
 
@@ -361,19 +469,19 @@ bool UMediaPlayer::Rewind()
 
 bool UMediaPlayer::Seek(const FTimespan& Time)
 {
-	return Player->Seek(Time);
+	return PlayerFacade->Seek(Time);
 }
 
 
 bool UMediaPlayer::SelectTrack(EMediaPlayerTrack TrackType, int32 TrackIndex)
 {
-	return Player->SelectTrack((EMediaTrackType)TrackType, TrackIndex);
+	return PlayerFacade->SelectTrack((EMediaTrackType)TrackType, TrackIndex);
 }
 
 
 void UMediaPlayer::SetDesiredPlayerName(FName PlayerName)
 {
-	Player->DesiredPlayerName = PlayerName;
+	PlayerFacade->DesiredPlayerName = PlayerName;
 }
 
 
@@ -381,85 +489,54 @@ bool UMediaPlayer::SetLooping(bool Looping)
 {
 	Loop = Looping;
 
-	return Player->SetLooping(Looping);
+	return PlayerFacade->SetLooping(Looping);
 }
 
 
-void UMediaPlayer::SetOverlays(UMediaOverlays* NewOverlays)
+bool UMediaPlayer::SetRate(float Rate)
 {
-	if (Overlays != nullptr)
-	{
-		Overlays->OnBeginDestroy().RemoveAll(this);
-	}
-
-	if (NewOverlays != nullptr)
-	{
-		NewOverlays->OnBeginDestroy().AddUObject(this, &UMediaPlayer::HandleMediaOverlaysBeginDestroy);
-	}
-
-	Player->SetOverlaySink(NewOverlays);
-
-	Overlays = NewOverlays;
+	return PlayerFacade->SetRate(Rate);
 }
 
 
-bool UMediaPlayer::SetRate(float Rate)
+bool UMediaPlayer::SetTrackFormat(EMediaPlayerTrack TrackType, int32 TrackIndex, int32 FormatIndex)
 {
-	return Player->SetRate(Rate);
+	return PlayerFacade->SetTrackFormat((EMediaTrackType)TrackType, TrackIndex, FormatIndex);
 }
 
 
-void UMediaPlayer::SetSoundWave(UMediaSoundWave* NewSoundWave)
+bool UMediaPlayer::SetVideoTrackFrameRate(int32 TrackIndex, int32 FormatIndex, float FrameRate)
 {
-	if (SoundWave != nullptr)
-	{
-		SoundWave->OnBeginDestroy().RemoveAll(this);
-	}
-
-	if (NewSoundWave != nullptr)
-	{
-		NewSoundWave->OnBeginDestroy().AddUObject(this, &UMediaPlayer::HandleMediaSoundWaveBeginDestroy);
-	}
-
-	Player->SetAudioSink(NewSoundWave);
-
-	SoundWave = NewSoundWave;
+	return PlayerFacade->SetVideoTrackFrameRate(TrackIndex, FormatIndex, FrameRate);
 }
 
-
-void UMediaPlayer::SetVideoTexture(UMediaTexture* NewTexture)
+bool UMediaPlayer::SetViewField(float Horizontal, float Vertical, bool Absolute)
 {
-	if (VideoTexture != nullptr)
-	{
-		VideoTexture->OnBeginDestroy().RemoveAll(this);
-	}
-
-	if (NewTexture != nullptr)
-	{
-		NewTexture->OnBeginDestroy().AddUObject(this, &UMediaPlayer::HandleMediaTextureBeginDestroy);
-	}
+	return PlayerFacade->SetViewField(Horizontal, Vertical, Absolute);
+}
 
-	Player->SetVideoSink(NewTexture);
 
-	VideoTexture = NewTexture;
+bool UMediaPlayer::SetViewRotation(const FRotator& Rotation, bool Absolute)
+{
+	return PlayerFacade->SetViewOrientation(FQuat(Rotation), Absolute);
 }
 
 
 bool UMediaPlayer::SupportsRate(float Rate, bool Unthinned) const
 {
-	return Player->SupportsRate(Rate, Unthinned);
+	return PlayerFacade->SupportsRate(Rate, Unthinned);
 }
 
 
 bool UMediaPlayer::SupportsScrubbing() const
 {
-	return Player->SupportsScrubbing();
+	return PlayerFacade->CanScrub();
 }
 
 
 bool UMediaPlayer::SupportsSeeking() const
 {
-	return Player->SupportsSeeking();
+	return PlayerFacade->CanSeek();
 }
 
 
@@ -492,13 +569,23 @@ void UMediaPlayer::ResumePIE()
 
 void UMediaPlayer::BeginDestroy()
 {
+	IMediaModule* MediaModule = FModuleManager::LoadModulePtr<IMediaModule>("Media");
+
+	if (MediaModule != nullptr)
+	{
+		MediaModule->GetClock().RemoveSink(PlayerFacade.ToSharedRef());
+		MediaModule->GetTicker().RemoveTickable(PlayerFacade.ToSharedRef());
+	}
+
+	PlayerFacade->Close();
+
 	Super::BeginDestroy();
+}
 
-	Close();
 
-	SetOverlays(nullptr);
-	SetSoundWave(nullptr);
-	SetVideoTexture(nullptr);
+bool UMediaPlayer::CanBeInCluster() const
+{
+	return false;
 }
 
 
@@ -508,23 +595,27 @@ FString UMediaPlayer::GetDesc()
 }
 
 
+void UMediaPlayer::PostDuplicate(bool bDuplicateForPIE)
+{
+	Super::PostDuplicate(bDuplicateForPIE);
+
+	PlayerGuid = FGuid::NewGuid();
+	PlayerFacade->SetGuid(PlayerGuid);
+}
+
+
 void UMediaPlayer::PostLoad()
 {
 	Super::PostLoad();
 
-	if (Overlays != nullptr)
-	{
-		Overlays->OnBeginDestroy().AddUObject(this, &UMediaPlayer::HandleMediaOverlaysBeginDestroy);
-	}
+	PlayerFacade->SetGuid(PlayerGuid);
 
-	if (SoundWave != nullptr)
-	{
-		SoundWave->OnBeginDestroy().AddUObject(this, &UMediaPlayer::HandleMediaSoundWaveBeginDestroy);
-	}
+	IMediaModule* MediaModule = FModuleManager::LoadModulePtr<IMediaModule>("Media");
 
-	if (VideoTexture != nullptr)
+	if (MediaModule != nullptr)
 	{
-		VideoTexture->OnBeginDestroy().AddUObject(this, &UMediaPlayer::HandleMediaTextureBeginDestroy);
+		MediaModule->GetClock().AddSink(PlayerFacade.ToSharedRef());
+		MediaModule->GetTicker().AddTickable(PlayerFacade.ToSharedRef());
 	}
 }
 
@@ -537,19 +628,7 @@ void UMediaPlayer::PostEditChangeProperty(FPropertyChangedEvent& PropertyChanged
 		? PropertyChangedEvent.Property->GetFName()
 		: NAME_None;
 
-	if (PropertyName == GET_MEMBER_NAME_CHECKED(UMediaPlayer, Overlays))
-	{
-		SetOverlays(Overlays);
-	}
-	else if (PropertyName == GET_MEMBER_NAME_CHECKED(UMediaPlayer, SoundWave))
-	{
-		SetSoundWave(SoundWave);
-	}
-	else if (PropertyName == GET_MEMBER_NAME_CHECKED(UMediaPlayer, VideoTexture))
-	{
-		SetVideoTexture(VideoTexture);
-	}
-	else if (PropertyName == GET_MEMBER_NAME_CHECKED(UMediaPlayer, Loop))
+	if (PropertyName == GET_MEMBER_NAME_CHECKED(UMediaPlayer, Loop))
 	{
 		SetLooping(Loop);
 	}
@@ -557,86 +636,12 @@ void UMediaPlayer::PostEditChangeProperty(FPropertyChangedEvent& PropertyChanged
 	Super::PostEditChangeProperty(PropertyChangedEvent);
 }
 
-
-void UMediaPlayer::PreEditChange(UProperty* PropertyAboutToChange)
-{
-	const FName PropertyName = (PropertyAboutToChange != nullptr)
-		? PropertyAboutToChange->GetFName()
-		: NAME_None;
-
-	if (PropertyName == GET_MEMBER_NAME_CHECKED(UMediaPlayer, Overlays))
-	{
-		SetOverlays(nullptr);
-	}
-	else if (PropertyName == GET_MEMBER_NAME_CHECKED(UMediaPlayer, SoundWave))
-	{
-		SetSoundWave(nullptr);
-	}
-	else if (PropertyName == GET_MEMBER_NAME_CHECKED(UMediaPlayer, VideoTexture))
-	{
-		SetVideoTexture(nullptr);
-	}
-
-	Super::PreEditChange(PropertyAboutToChange);
-}
-
 #endif
 
 
-/* FTickerObjectBase interface
- *****************************************************************************/
-
-bool UMediaPlayer::Tick(float DeltaTime)
-{
-	Player->TickPlayer(DeltaTime);
-
-	typedef TWeakPtr<FMediaPlayerBase, ESPMode::ThreadSafe> FMediaPlayerBasePtr;
-
-	ENQUEUE_UNIQUE_RENDER_COMMAND_TWOPARAMETER(MediaPlayerTickRender,
-		FMediaPlayerBasePtr, PlayerPtr, Player,
-		float, DeltaTime, DeltaTime,
-		{
-			auto PinnedPlayer = PlayerPtr.Pin();
-			if (PinnedPlayer.IsValid())
-			{
-				PinnedPlayer->TickVideo(DeltaTime);
-			}
-		});
-
-	return true;
-}
-
-
 /* UMediaPlayer callbacks
  *****************************************************************************/
 
-void UMediaPlayer::HandleMediaOverlaysBeginDestroy(UMediaOverlays& DestroyedOverlays)
-{
-	if (&DestroyedOverlays == Overlays)
-	{
-		Player->SetOverlaySink(nullptr);
-	}
-}
-
-
-void UMediaPlayer::HandleMediaSoundWaveBeginDestroy(UMediaSoundWave& DestroyedSoundWave)
-{
-	if (&DestroyedSoundWave == SoundWave)
-	{
-		Player->SetAudioSink(nullptr);
-	}
-}
-
-
-void UMediaPlayer::HandleMediaTextureBeginDestroy(UMediaTexture& DestroyedMediaTexture)
-{
-	if (&DestroyedMediaTexture == VideoTexture)
-	{
-		Player->SetVideoSink(nullptr);
-	}
-}
-
-
 void UMediaPlayer::HandlePlayerMediaEvent(EMediaEvent Event)
 {
 	MediaEvent.Broadcast(Event);
@@ -648,33 +653,22 @@ void UMediaPlayer::HandlePlayerMediaEvent(EMediaEvent Event)
 		break;
 
 	case EMediaEvent::MediaOpened:
-		Player->SetLooping(Loop);
+		PlayerFacade->SetCacheWindow(CacheAhead, FApp::IsGame() ? CacheBehindGame : CacheBehind);
+		PlayerFacade->SetLooping(Loop);
+		PlayerFacade->SetViewField(HorizontalFieldOfView, VerticalFieldOfView, true);
+		PlayerFacade->SetViewOrientation(FQuat(ViewRotation), true);
 
-		if (Overlays != nullptr)
-		{
-			Player->SetOverlaySink(Overlays);
-		}
-
-		if (SoundWave != nullptr)
-		{
-			Player->SetAudioSink(SoundWave);
-		}
-
-		if (VideoTexture != nullptr)
-		{
-			Player->SetVideoSink(VideoTexture);
-		}
-
-		OnMediaOpened.Broadcast(Player->GetUrl());
+		OnMediaOpened.Broadcast(PlayerFacade->GetUrl());
 
-		if (PlayOnOpen)
+		if (PlayOnOpen || PlayOnNext)
 		{
+			PlayOnNext = false;
 			Play();
 		}
 		break;
 
 	case EMediaEvent::MediaOpenFailed:
-		OnMediaOpenFailed.Broadcast(Player->GetUrl());
+		OnMediaOpenFailed.Broadcast(PlayerFacade->GetUrl());
 
 		if (!Loop && (Playlist != nullptr))
 		{
@@ -685,8 +679,9 @@ void UMediaPlayer::HandlePlayerMediaEvent(EMediaEvent Event)
 	case EMediaEvent::PlaybackEndReached:
 		OnEndReached.Broadcast();
 
-		if (!Loop && (Playlist != nullptr))
+		if (!Loop && (Playlist != nullptr) && (Playlist->Num() > 1))
 		{
+			PlayOnNext = true;
 			Next();
 		}
 		break;
@@ -698,5 +693,9 @@ void UMediaPlayer::HandlePlayerMediaEvent(EMediaEvent Event)
 	case EMediaEvent::PlaybackSuspended:
 		OnPlaybackSuspended.Broadcast();
 		break;
+
+	case EMediaEvent::SeekCompleted:
+		OnSeekCompleted.Broadcast();
+		break;
 	}
 }
diff --git a/Engine/Source/Runtime/MediaAssets/Private/Assets/MediaPlaylist.cpp b/Engine/Source/Runtime/MediaAssets/Private/Assets/MediaPlaylist.cpp
--- a/Engine/Source/Runtime/MediaAssets/Private/Assets/MediaPlaylist.cpp
+++ b/Engine/Source/Runtime/MediaAssets/Private/Assets/MediaPlaylist.cpp
@@ -1,11 +1,70 @@
 // Copyright 1998-2017 Epic Games, Inc. All Rights Reserved.
 
 #include "MediaPlaylist.h"
+#include "MediaAssetsPrivate.h"
+
+#include "Misc/Paths.h"
+#include "UObject/Package.h"
+#include "UObject/UObjectGlobals.h"
+
+#include "StreamMediaSource.h"
 
 
 /* UMediaPlaylist interface
  *****************************************************************************/
 
+bool UMediaPlaylist::Add(UMediaSource* MediaSource)
+{
+	if (MediaSource == nullptr)
+	{
+		return false;
+	}
+
+	Items.Add(MediaSource);
+
+	return true;
+}
+
+
+bool UMediaPlaylist::AddFile(const FString& FilePath)
+{
+	FString FullPath;
+
+	if (FPaths::IsRelative(FilePath))
+	{
+		FullPath = FPaths::ConvertRelativePathToFull(FilePath);
+	}
+	else
+	{
+		FullPath = FilePath;
+		FPaths::NormalizeFilename(FullPath);
+	}
+
+	return AddUrl(FString(TEXT("file://")) + FullPath);
+}
+
+
+bool UMediaPlaylist::AddUrl(const FString& Url)
+{
+	if (Url.IsEmpty())
+	{
+		return false;
+	}
+
+	auto MediaSource = NewObject<UStreamMediaSource>(GetTransientPackage(), NAME_None, RF_Transactional | RF_Transient);
+	MediaSource->StreamUrl = Url;
+
+	if (!MediaSource->Validate())
+	{
+		UE_LOG(LogMediaAssets, Error, TEXT("Failed to validate URL %s"), *Url);
+
+		return false;
+	}
+
+	return Add(MediaSource);
+}
+
+
 UMediaSource* UMediaPlaylist::Get(int32 Index)
 {
 	if (!Items.IsValidIndex(Index))
@@ -19,44 +78,102 @@ UMediaSource* UMediaPlaylist::Get(int32 Index)
 
 UMediaSource* UMediaPlaylist::GetNext(int32& InOutIndex)
 {
-	if (Items.Num() == 0)
+	if (Items.Num() > 0)
 	{
-		InOutIndex = INDEX_NONE;
+		InOutIndex = FMath::Clamp(InOutIndex, (int32)INDEX_NONE, Items.Num() - 1) + 1;
 
-		return nullptr;
+		if (Loop)
+		{
+			InOutIndex %= Items.Num();
+		}
+
+		if (InOutIndex < Items.Num())
+		{
+			return Items[InOutIndex];
+		}
 	}
 
-	InOutIndex = (InOutIndex != INDEX_NONE) ? (InOutIndex + 1) % Items.Num() : 0;
+	InOutIndex = INDEX_NONE;
 
-	return Items[InOutIndex];
+	return nullptr;
 }
 
 
 UMediaSource* UMediaPlaylist::GetPrevious(int32& InOutIndex)
 {
-	if (Items.Num() == 0)
+	if (Items.Num() > 0)
+	{
+		InOutIndex = FMath::Clamp(InOutIndex, 0, Items.Num()) - 1;
+
+		if (Loop)
+		{
+			InOutIndex += Items.Num();
+			InOutIndex %= Items.Num();
+		}
+
+		if (InOutIndex >= 0)
+		{
+			return Items[InOutIndex];
+		}
+	}
+	else
 	{
 		InOutIndex = INDEX_NONE;
+	}
+
+	return nullptr;
+}
+
+
+UMediaSource* UMediaPlaylist::GetRandom(int32& OutIndex)
+{
+	if (Items.Num() == 0)
+	{
+		OutIndex = INDEX_NONE;
 
 		return nullptr;
 	}
 
-	InOutIndex = (InOutIndex != INDEX_NONE) ? (InOutIndex + Items.Num() - 1) % Items.Num() : 0;
+	OutIndex = FMath::RandHelper(Items.Num() - 1);
 
-	return Items[InOutIndex];
+	return Items[OutIndex];
 }
 
 
-UMediaSource* UMediaPlaylist::GetRandom(int32& InOutIndex)
+void UMediaPlaylist::Insert(UMediaSource* MediaSource, int32 Index)
 {
-	if (Items.Num() == 0)
+	Index = FMath::Clamp(Index, 0, Items.Num());
+	Items.Insert(MediaSource, Index);
+}
+
+
+bool UMediaPlaylist::Remove(UMediaSource* MediaSource)
+{
+	return (Items.Remove(MediaSource) > 0);
+}
+
+
+bool UMediaPlaylist::RemoveAt(int32 Index)
+{
+	if (!Items.IsValidIndex(Index))
 	{
-		InOutIndex = INDEX_NONE;
+		return false;
+	}
 
-		return nullptr;
+	Items.RemoveAt(Index);
+
+	return true;
+}
+
+
+bool UMediaPlaylist::Replace(int32 Index, UMediaSource* Replacement)
+{
+	if (!Items.IsValidIndex(Index))
+	{
+		return false;
 	}
 
-	InOutIndex = FMath::RandHelper(Items.Num() - 1);
+	Items[Index] = Replacement;
 
-	return Items[InOutIndex];
+	return true;
 }
diff --git a/Engine/Source/Runtime/MediaAssets/Private/Assets/MediaSoundComponent.cpp b/Engine/Source/Runtime/MediaAssets/Private/Assets/MediaSoundComponent.cpp
new file mode 100644
--- /dev/null
+++ b/Engine/Source/Runtime/MediaAssets/Private/Assets/MediaSoundComponent.cpp
@@ -0,0 +1,143 @@
+// Copyright 1998-2017 Epic Games, Inc. All Rights Reserved.
+
+#include "MediaSoundComponent.h"
+
+#include "IMediaAudioSample.h"
+#include "IMediaPlayer.h"
+#include "MediaAudioResampler.h"
+#include "Misc/ScopeLock.h"
+
+#include "MediaPlayer.h"
+#include "MediaPlayerFacade.h"
+
+
+/* UMediaSoundComponent structors
+ *****************************************************************************/
+
+UMediaSoundComponent::UMediaSoundComponent(const FObjectInitializer& ObjectInitializer)
+	: Super(ObjectInitializer)
+	, Channels(EMediaSoundChannels::Stereo)
+	, Resampler(new FMediaAudioResampler)
+{
+	PrimaryComponentTick.bCanEverTick = true;
+}
+
+
+UMediaSoundComponent::~UMediaSoundComponent()
+{
+	delete Resampler;
+}
+
+
+/* UMediaSoundComponent interface
+ *****************************************************************************/
+
+void UMediaSoundComponent::UpdatePlayer()
+{
+	if (MediaPlayer == nullptr)
+	{
+		SampleQueue.Reset();
+		return;
+	}
+
+	// create a new sample queue if the player changed
+	TSharedRef<FMediaPlayerFacade, ESPMode::ThreadSafe> PlayerFacade = MediaPlayer->GetPlayerFacade();
+
+	if (PlayerFacade != CurrentPlayerFacade)
+	{
+		{
+			const auto NewSampleQueue = MakeShared<FMediaAudioSampleQueue, ESPMode::ThreadSafe>();
+
+			FScopeLock Lock(&CriticalSection);
+			SampleQueue = NewSampleQueue;
+		}
+
+		PlayerFacade->AddAudioSampleSink(SampleQueue.ToSharedRef());
+		CurrentPlayerFacade = PlayerFacade;
+	}
+
+	check(SampleQueue.IsValid());
+}
+
+
+/* UActorComponent interface
+ *****************************************************************************/
+
+void UMediaSoundComponent::TickComponent(float DeltaTime, enum ELevelTick TickType, FActorComponentTickFunction *ThisTickFunction)
+{
+	Super::TickComponent(DeltaTime, TickType, ThisTickFunction);
+
+	UpdatePlayer();
+}
+
+
+/* USynthComponent interface
+ *****************************************************************************/
+
+void UMediaSoundComponent::Activate(bool bReset)
+{
+	if (bReset || ShouldActivate())
+	{
+		SetComponentTickEnabled(true);
+	}
+
+	Super::Activate(bReset);
+}
+
+
+void UMediaSoundComponent::Deactivate()
+{
+	if (!ShouldActivate())
+	{
+		SetComponentTickEnabled(false);
+	}
+
+	Super::Deactivate();
+}
+
+
+/* USynthComponent interface
+ *****************************************************************************/
+
+void UMediaSoundComponent::Init(const int32 SampleRate)
+{
+	Super::Init(SampleRate);
+
+	if (Channels == EMediaSoundChannels::Mono)
+	{
+		NumChannels = 1;
+	}
+	else //if (Channels == EMediaSoundChannels::Stereo)
+	{
+		NumChannels = 2;
+	}/*
+	else
+	{
+		NumChannels = 8;
+	}*/
+
+	Resampler->Initialize(NumChannels, SampleRate);
+}
+
+
+void UMediaSoundComponent::OnGenerateAudio(TArray<float>& OutAudio)
+{
+	TSharedPtr<FMediaPlayerFacade, ESPMode::ThreadSafe> PinnedPlayerFacade;
+	TSharedPtr<FMediaAudioSampleQueue, ESPMode::ThreadSafe> PinnedSampleQueue;
+	{
+		FScopeLock Lock(&CriticalSection);
+
+		PinnedPlayerFacade = CurrentPlayerFacade.Pin();
+		PinnedSampleQueue = SampleQueue;
+	}
+
+	if (PinnedPlayerFacade.IsValid() && PinnedPlayerFacade->IsPlaying() && PinnedSampleQueue.IsValid())
+	{
+		const uint32 FramesRequested = OutAudio.Num() / NumChannels;
+		const uint32 FramesWritten = Resampler->Generate(OutAudio.GetData(), FramesRequested, PinnedPlayerFacade->GetRate(), PinnedPlayerFacade->GetTime(), *PinnedSampleQueue);
+	}
+	else
+	{
+		Resampler->Flush();
+	}
+}
diff --git a/Engine/Source/Runtime/MediaAssets/Private/Assets/MediaTexture.cpp b/Engine/Source/Runtime/MediaAssets/Private/Assets/MediaTexture.cpp
--- a/Engine/Source/Runtime/MediaAssets/Private/Assets/MediaTexture.cpp
+++ b/Engine/Source/Runtime/MediaAssets/Private/Assets/MediaTexture.cpp
@@ -1,12 +1,55 @@
 // Copyright 1998-2017 Epic Games, Inc. All Rights Reserved.
 
 #include "MediaTexture.h"
-#include "Misc/ScopeLock.h"
-#include "RenderUtils.h"
 #include "MediaAssetsPrivate.h"
+
+#include "ExternalTexture.h"
+#include "IMediaClock.h"
+#include "IMediaClockSink.h"
+#include "IMediaModule.h"
+#include "MediaPlayerFacade.h"
+#include "Modules/ModuleManager.h"
+#include "RenderUtils.h"
+#include "RenderingThread.h"
+#include "UObject/WeakObjectPtrTemplates.h"
+
+#include "MediaPlayer.h"
 #include "Misc/MediaTextureResource.h"
 
 
+/* Local helpers
+ *****************************************************************************/
+
+/**
+ * Media clock sink for media textures.
+ */
+class FMediaTextureClockSink
+	: public IMediaClockSink
+{
+public:
+
+	FMediaTextureClockSink(UMediaTexture& InOwner)
+		: Owner(&InOwner)
+	{ }
+
+	virtual ~FMediaTextureClockSink() { }
+
+public:
+
+	virtual void TickRender(FTimespan DeltaTime, FTimespan Timecode) override
+	{
+		if (Owner.IsValid())
+		{
+			Owner->TickResource(Timecode);
+		}
+	}
+
+private:
+
+	TWeakObjectPtr<UMediaTexture> Owner;
+};
+
+
 /* UMediaTexture structors
  *****************************************************************************/
 
@@ -14,28 +57,21 @@ UMediaTexture::UMediaTexture(const FObjectInitializer& ObjectInitializer)
 	: Super(ObjectInitializer)
 	, AddressX(TA_Clamp)
 	, AddressY(TA_Clamp)
+	, AutoClear(true)
 	, ClearColor(FLinearColor::Black)
-	, SinkBufferDim(FIntPoint(1, 1))
-	, SinkOutputDim(FIntPoint(1, 1))
-	, SinkFormat(EMediaTextureSinkFormat::CharBGRA)
-	, SinkMode(EMediaTextureSinkMode::Unbuffered)
+	, MediaPlayer(nullptr)
+	, Dimensions(FIntPoint::ZeroValue)
+	, Size(0)
 {
 	NeverStream = true;
 }
 
 
-/* IMediaTextureSink interface
+/* UMediaTexture interface
  *****************************************************************************/
 
 float UMediaTexture::GetAspectRatio() const
 {
-	if (Resource == nullptr)
-	{
-		return 0.0f;
-	}
-
-	const FIntPoint Dimensions = ((FMediaTextureResource*)Resource)->GetSizeXY();
-
 	if (Dimensions.Y == 0)
 	{
 		return 0.0f;
@@ -47,13 +83,13 @@ float UMediaTexture::GetAspectRatio() const
 
 int32 UMediaTexture::GetHeight() const
 {
-	return (Resource != nullptr) ? (int32)Resource->GetSizeY() : 0;
+	return Dimensions.Y;
 }
 
 
 int32 UMediaTexture::GetWidth() const
 {
-	return (Resource != nullptr) ? (int32)Resource->GetSizeX() : 0;
+	return Dimensions.X;
 }
 
 
@@ -62,33 +98,42 @@ int32 UMediaTexture::GetWidth() const
 
 FTextureResource* UMediaTexture::CreateResource()
 {
-	return new FMediaTextureResource(*this, ClearColor, SinkOutputDim, SinkFormat, SinkMode);
+	if (!ClockSink.IsValid())
+	{
+		IMediaModule* MediaModule = FModuleManager::LoadModulePtr<IMediaModule>("Media");
+
+		if (MediaModule != nullptr)
+		{
+			ClockSink = MakeShared<FMediaTextureClockSink, ESPMode::ThreadSafe>(*this);
+			MediaModule->GetClock().AddSink(ClockSink.ToSharedRef());
+		}
+	}
+
+	return new FMediaTextureResource(*this, Dimensions, Size);
 }
 
 
-EMaterialValueType UMediaTexture::GetMaterialType()
+EMaterialValueType UMediaTexture::GetMaterialType() const
 {
-	return MCT_Texture2D;
+	return MCT_TextureExternal;
 }
 
 
 float UMediaTexture::GetSurfaceWidth() const
 {
-	return (Resource != nullptr) ? Resource->GetSizeX() : 0.0f;
+	return Dimensions.X;
 }
 
 
 float UMediaTexture::GetSurfaceHeight() const
 {
-	return (Resource != nullptr) ? Resource->GetSizeY() : 0.0f;
+	return Dimensions.Y;
 }
 
 
-void UMediaTexture::UpdateResource()
+FGuid UMediaTexture::GetExternalTextureGuid() const
 {
-	FScopeLock Lock(&CriticalSection);
-
-	Super::UpdateResource();
+	return MediaPlayer ? MediaPlayer->GetGuid() : FGuid();
 }
 
 
@@ -97,15 +142,27 @@ void UMediaTexture::UpdateResource()
 
 void UMediaTexture::BeginDestroy()
 {
-	Super::BeginDestroy();
+	if (ClockSink.IsValid())
+	{
+		IMediaModule* MediaModule = FModuleManager::LoadModulePtr<IMediaModule>("Media");
+
+		if (MediaModule != nullptr)
+		{
+			MediaModule->GetClock().RemoveSink(ClockSink.ToSharedRef());
+		}
+
+		ClockSink.Reset();
+	}
+
+	UnregisterPlayerGuid();
 
-	BeginDestroyEvent.Broadcast(*this);
+	Super::BeginDestroy();
 }
 
 
 FString UMediaTexture::GetDesc()
 {
-	return FString::Printf(TEXT("%dx%d [%s]"), GetSurfaceWidth(),  GetSurfaceHeight(), GPixelFormats[PF_B8G8R8A8].Name);
+	return FString::Printf(TEXT("%ix%i [%s]"), Dimensions.X,  Dimensions.Y, GPixelFormats[PF_B8G8R8A8].Name);
 }
 
 
@@ -113,10 +170,7 @@ void UMediaTexture::GetResourceSizeEx(FResourceSizeEx& CumulativeResourceSize)
 {
 	Super::GetResourceSizeEx(CumulativeResourceSize);
 
-	if (Resource)
-	{
-		((FMediaTextureResource*)Resource)->GetResourceSizeEx(CumulativeResourceSize);
-	}
+	CumulativeResourceSize.AddUnknownMemoryBytes(Size);
 }
 
 
@@ -124,162 +178,120 @@ void UMediaTexture::GetResourceSizeEx(FResourceSizeEx& CumulativeResourceSize)
 
 void UMediaTexture::PostEditChangeProperty(FPropertyChangedEvent& PropertyChangedEvent)
 {
-	const FName PropertyName = (PropertyChangedEvent.Property != nullptr)
-		? PropertyChangedEvent.Property->GetFName()
-		: NAME_None;
-
-	if ((PropertyName == GET_MEMBER_NAME_CHECKED(UMediaTexture, ClearColor)) &&
-		(PropertyChangedEvent.ChangeType != EPropertyChangeType::ValueSet))
-	{
-		UObject::PostEditChangeProperty(PropertyChangedEvent);
-	}
-	else
+	static const FName AddressXName = GET_MEMBER_NAME_CHECKED(UMediaTexture, AddressX);
+	static const FName AddressYName = GET_MEMBER_NAME_CHECKED(UMediaTexture, AddressY);
+	static const FName AutoClearName = GET_MEMBER_NAME_CHECKED(UMediaTexture, AutoClear);
+	static const FName ClearColorName = GET_MEMBER_NAME_CHECKED(UMediaTexture, ClearColor);
+	static const FName MediaPlayerName = GET_MEMBER_NAME_CHECKED(UMediaTexture, MediaPlayer);
+
+	UProperty* PropertyThatChanged = PropertyChangedEvent.Property;
+	
+	if (PropertyThatChanged == nullptr)
 	{
 		Super::PostEditChangeProperty(PropertyChangedEvent);
-	}
-}
-
-
-void UMediaTexture::PreEditChange(UProperty* PropertyAboutToChange)
-{
-	Super::PreEditChange(PropertyAboutToChange);
 
-	FlushRenderingCommands();
-}
-
-#endif // WITH_EDITOR
-
-
-/* IMediaTextureSink interface
- *****************************************************************************/
+		return;
+	}
 
-void* UMediaTexture::AcquireTextureSinkBuffer()
-{
-	FScopeLock Lock(&CriticalSection);
+	const FName PropertyName = PropertyThatChanged->GetFName();
 
-	return (Resource != nullptr) ? ((FMediaTextureResource*)Resource)->AcquireBuffer() : nullptr;
-}
+	// don't update resource for these properties
+	if ((PropertyName == AutoClearName) ||
+		(PropertyName == ClearColorName) ||
+		(PropertyName == MediaPlayerName))
+	{
+		UObject::PostEditChangeProperty(PropertyChangedEvent);
 
+		return;
+	}
 
-void UMediaTexture::DisplayTextureSinkBuffer(FTimespan /*Time*/)
-{
-	FScopeLock Lock(&CriticalSection);
+	Super::PostEditChangeProperty(PropertyChangedEvent);
 
-	if (Resource != nullptr)
+	// notify materials for these properties
+	if ((PropertyName == AddressXName) ||
+		(PropertyName == AddressYName))
 	{
-		((FMediaTextureResource*)Resource)->DisplayBuffer();
+		NotifyMaterials();
 	}
 }
 
+#endif // WITH_EDITOR
 
-FIntPoint UMediaTexture::GetTextureSinkDimensions() const
-{
-	FScopeLock Lock(&CriticalSection);
-
-	return (Resource != nullptr) ? ((FMediaTextureResource*)Resource)->GetSizeXY() : FIntPoint::ZeroValue;
-}
-
-
-EMediaTextureSinkFormat UMediaTexture::GetTextureSinkFormat() const
-{
-	return SinkFormat;
-}
-
-
-EMediaTextureSinkMode UMediaTexture::GetTextureSinkMode() const
-{
-	return SinkMode;
-}
-
-
-FRHITexture* UMediaTexture::GetTextureSinkTexture()
-{
-	FScopeLock Lock(&CriticalSection);
-
-	return (Resource != nullptr) ? ((FMediaTextureResource*)Resource)->GetTexture() : nullptr;
-}
 
+/* UMediaTexture implementation
+ *****************************************************************************/
 
-bool UMediaTexture::InitializeTextureSink(FIntPoint OutputDim, FIntPoint BufferDim, EMediaTextureSinkFormat Format, EMediaTextureSinkMode Mode)
+void UMediaTexture::TickResource(FTimespan Timecode)
 {
-	UE_LOG(LogMediaAssets, Verbose, TEXT("MediaTexture initializing sink with %ix%i output and %ix%i buffer as %s."), OutputDim.X, OutputDim.Y, BufferDim.X, BufferDim.Y, (Mode == EMediaTextureSinkMode::Buffered) ? TEXT("Buffered") : TEXT("Unbuffered"));
-
-	SinkBufferDim = BufferDim;
-	SinkOutputDim = OutputDim;
-	SinkFormat = Format;
-	SinkMode = Mode;
-
-	FScopeLock Lock(&CriticalSection);
-
-	if (Resource == nullptr)
+	if ((MediaPlayer == nullptr) || (Resource == nullptr))
 	{
-		return false;
+		SampleQueue.Reset();
+		return;
 	}
 
-	((FMediaTextureResource*)Resource)->InitializeBuffer(OutputDim, BufferDim, Format, Mode);
+	// create a new sample queue if the player changed
+	TSharedRef<FMediaPlayerFacade, ESPMode::ThreadSafe> PlayerFacade = MediaPlayer->GetPlayerFacade();
 
-	return true;
-}
+	if (PlayerFacade != CurrentPlayerFacade)
+	{
+		SampleQueue = MakeShared<FMediaTextureSampleQueue, ESPMode::ThreadSafe>();
+		PlayerFacade->AddVideoSampleSink(SampleQueue.ToSharedRef());
+		CurrentPlayerFacade = PlayerFacade;
+	}
 
+	check(SampleQueue.IsValid());
 
-void UMediaTexture::ReleaseTextureSinkBuffer()
-{
-	FScopeLock Lock(&CriticalSection);
+	// unregister previous external texture GUID if needed
+	const FGuid PlayerGuid = MediaPlayer->GetGuid();
 
-	if (Resource != nullptr)
+	if (PlayerGuid != LastPlayerGuid)
 	{
-		((FMediaTextureResource*)Resource)->ReleaseBuffer();
+		UnregisterPlayerGuid();
+		LastPlayerGuid = PlayerGuid;
 	}
-}
 
+	// retain the last rendered frame if player is inactive
+	const bool PlayerActive = MediaPlayer->IsPaused() || MediaPlayer->IsPlaying() || MediaPlayer->IsPreparing();
 
-void UMediaTexture::ShutdownTextureSink()
-{
-	UE_LOG(LogMediaAssets, Verbose, TEXT("MediaTexture shutting down sink."));
-
-	if (ClearColor.A == 0.0f)
+	if (!PlayerActive && !AutoClear)
 	{
 		return;
 	}
 
-	// reset to 1x1 clear color
-	SinkBufferDim = FIntPoint(1, 1);
-	SinkOutputDim = FIntPoint(1, 1);
-	SinkFormat = EMediaTextureSinkFormat::CharBGRA;
-	SinkMode = EMediaTextureSinkMode::Unbuffered;
-
-	FScopeLock Lock(&CriticalSection);
-
-	if (Resource != nullptr)
+	// issue a render command to render the current sample
+	FMediaTextureResource::FRenderParams RenderParams;
 	{
-		((FMediaTextureResource*)Resource)->InitializeBuffer(SinkOutputDim, SinkBufferDim, SinkFormat, SinkMode);
+		RenderParams.ClearColor = ClearColor;
+		RenderParams.PlayerGuid = PlayerGuid;
+		RenderParams.Rate = MediaPlayer->GetRate();
+		RenderParams.SrgbOutput = SRGB;
+		RenderParams.Time = MediaPlayer->GetTime();
+
+		if (PlayerActive)
+		{
+			RenderParams.SampleSource = SampleQueue;
+		}
 	}
-}
-
 
-bool UMediaTexture::SupportsTextureSinkFormat(EMediaTextureSinkFormat Format) const
-{
-	return true; // all formats are supported
+	ENQUEUE_UNIQUE_RENDER_COMMAND_TWOPARAMETER(MediaTextureResourceRender,
+		FMediaTextureResource*, ResourceParam, (FMediaTextureResource*)Resource,
+		FMediaTextureResource::FRenderParams, RenderParamsParam, RenderParams,
+		{
+			ResourceParam->Render(RenderParamsParam);
+		});
 }
 
 
-void UMediaTexture::UpdateTextureSinkBuffer(const uint8* Data, uint32 Pitch)
+void UMediaTexture::UnregisterPlayerGuid()
 {
-	FScopeLock Lock(&CriticalSection);
-
-	if (Resource != nullptr)
+	if (!LastPlayerGuid.IsValid())
 	{
-		((FMediaTextureResource*)Resource)->UpdateBuffer(Data, Pitch);
+		return;
 	}
-}
-
-
-void UMediaTexture::UpdateTextureSinkResource(FRHITexture* RenderTarget, FRHITexture* ShaderResource)
-{
-	FScopeLock Lock(&CriticalSection);
 
-	if (Resource != nullptr)
-	{
-		((FMediaTextureResource*)Resource)->UpdateTextures(RenderTarget, ShaderResource);
-	}
+	ENQUEUE_UNIQUE_RENDER_COMMAND_ONEPARAMETER(MediaTextureUnregisterPlayerGuid,
+		FGuid, PlayerGuid, LastPlayerGuid,
+		{
+			FExternalTextureRegistry::Get().UnregisterExternalTexture(PlayerGuid);
+		});
 }
diff --git a/Engine/Source/Runtime/MediaAssets/Private/Assets/PlatformMediaSource.cpp b/Engine/Source/Runtime/MediaAssets/Private/Assets/PlatformMediaSource.cpp
--- a/Engine/Source/Runtime/MediaAssets/Private/Assets/PlatformMediaSource.cpp
+++ b/Engine/Source/Runtime/MediaAssets/Private/Assets/PlatformMediaSource.cpp
@@ -59,7 +59,7 @@ void UPlatformMediaSource::Serialize(FArchive& Ar)
 		}
 		else
 		{
-			Ar << PlatformMediaSources;
+		Ar << PlatformMediaSources;
 		}
 #else
 		Ar << MediaSource;
@@ -83,13 +83,13 @@ bool UPlatformMediaSource::Validate() const
 	for (auto PlatformNameMediaSourcePair : PlatformMediaSources)
 	{
 		UMediaSource* PlatformMediaSource = PlatformNameMediaSourcePair.Value;
+
 		if (PlatformMediaSource != nullptr && PlatformMediaSource->Validate() == false)
 		{
-			// If any platform is specified but doesn't validate, the entire platform media source is not valid.
 			return false;
 		}
 	}
-	// If there are any platform media sources specified they will have been validated above.
+
 	return PlatformMediaSources.Num() > 0;
 #else
 	return (MediaSource != nullptr) ? MediaSource->Validate() : false;
