Copying //UE4/Dev-Sequencer to //UE4/Main (Source: //UE4/Dev-Sequencer @ 3617642)

diff --git a/Engine/Source/Runtime/MediaAssets/Private/Misc/MediaTextureResource.cpp b/Engine/Source/Runtime/MediaAssets/Private/Misc/MediaTextureResource.cpp
--- a/Engine/Source/Runtime/MediaAssets/Private/Misc/MediaTextureResource.cpp
+++ b/Engine/Source/Runtime/MediaAssets/Private/Misc/MediaTextureResource.cpp
@@ -1,211 +1,224 @@
 // Copyright 1998-2017 Epic Games, Inc. All Rights Reserved.
 
 #include "Misc/MediaTextureResource.h"
-#include "RenderingThread.h"
-#include "Shader.h"
-#include "IMediaTextureSink.h"
-#include "MediaTexture.h"
+#include "MediaAssetsPrivate.h"
+
 #include "DeviceProfiles/DeviceProfile.h"
-#include "StaticBoundShaderState.h"
-#include "SceneUtils.h"
 #include "DeviceProfiles/DeviceProfileManager.h"
-#include "RHIStaticStates.h"
+#include "ExternalTexture.h"
+#include "IMediaPlayer.h"
+#include "IMediaSamples.h"
+#include "IMediaTextureSample.h"
+#include "MediaPlayerFacade.h"
+#include "MediaSampleSource.h"
 #include "MediaShaders.h"
-#include "MediaAssetsPrivate.h"
 #include "PipelineStateCache.h"
+#include "SceneUtils.h"
+#include "Shader.h"
+#include "StaticBoundShaderState.h"
+#include "RenderUtils.h"
+#include "RHIStaticStates.h"
+#include "ExternalTexture.h"
 
-/* FMediaTextureResource structors
- *****************************************************************************/
-
-FMediaTextureResource::FMediaTextureResource(UMediaTexture& InOwner, const FLinearColor& InClearColor, const FIntPoint& InOutputDimensions, EMediaTextureSinkFormat InSinkFormat, EMediaTextureSinkMode InSinkMode)
-	: Owner(InOwner)
-	, TripleBuffer(BufferResources)
-	, BufferClearColor(InClearColor)
-	, BufferDimensions(InOutputDimensions)
-	, BufferPitch(0)
-	, CachedResourceSizeBytes(0)
-	, OutputDimensions(InOutputDimensions)
-	, RequiresConversion(false)
-	, SinkFormat(InSinkFormat)
-	, SinkMode(InSinkMode)
-	, State(EState::ShutDown)
-{ }
+#include "MediaTexture.h"
 
 
-FMediaTextureResource::~FMediaTextureResource()
-{
-	FlushRenderingCommands();
-}
+#define MEDIATEXTURERESOURCE_TRACE_RENDER 0
 
 
-/* FMediaTextureResource interface
+/* Local helpers
  *****************************************************************************/
 
-void* FMediaTextureResource::AcquireBuffer()
-{
-	if (State != EState::Initialized)
-	{
-		return nullptr;
-	}
-
-	if (SinkMode == EMediaTextureSinkMode::Buffered)
-	{
-		return TripleBuffer.GetWriteBuffer().LockedData;
-	}
-
-	check(IsInRenderingThread());
-
-	uint32 OutStride = 0;
-	TRefCountPtr<FRHITexture2D>& Texture = RequiresConversion
-		? BufferResources[2].RenderTarget
-		: RenderTargetTextureRHI;
-
-	return RHILockTexture2D(Texture.GetReference(), 0, RLM_WriteOnly, OutStride, false /*bLockWithinMiptail*/, false /*bFlushRHIThread*/);
-}
-
-
-void FMediaTextureResource::DisplayBuffer()
+namespace MediaTextureResource
 {
-	if (State != EState::Initialized)
-	{
-		return;
+	/**
+	 * Get the pixel format for a given sample.
+	 *
+	 * @param Sample The sample.
+	 * @return The sample's pixel format.
+	 */
+	EPixelFormat GetPixelFormat(const TSharedPtr<IMediaTextureSample, ESPMode::ThreadSafe>& Sample)
+	{
+		switch (Sample->GetFormat())
+		{
+		case EMediaTextureSampleFormat::CharAYUV:
+		case EMediaTextureSampleFormat::CharBGRA:
+		case EMediaTextureSampleFormat::CharBMP:
+		case EMediaTextureSampleFormat::CharUYVY:
+		case EMediaTextureSampleFormat::CharYUY2:
+		case EMediaTextureSampleFormat::CharYVYU:
+			return PF_B8G8R8A8;
+
+		case EMediaTextureSampleFormat::CharNV12:
+		case EMediaTextureSampleFormat::CharNV21:
+			return PF_G8;
+
+		case EMediaTextureSampleFormat::FloatRGB:
+			return PF_FloatRGB;
+
+		case EMediaTextureSampleFormat::FloatRGBA:
+			return PF_FloatRGBA;
+
+		default:
+			return PF_Unknown;
+		}
 	}
-	
-	// RenderThreadTasks will be consumed by ENQUEUE_RENDER_COMMAND in FMediaTextureResource::Tick.
-	// Since ENQUEUE_RENDER_COMMAND will decide which thread to run the tasks on, we should not
-	// make decision ourself to avoid conflict.
-	RenderThreadTasks.Enqueue([this]() {
-		SwapResource();
-	});
-}
 
 
-FRHITexture* FMediaTextureResource::GetTexture()
-{
-	if (State != EState::Initialized)
+	/**
+	 * Check whether the given sample requires a conversion shader.
+	 *
+	 * @param Sample The sample to check.
+	 * @param SrgbOutput Whether the output is expected in sRGB color space.
+	 * @return true if conversion is required, false otherwise.
+	 * @see RequiresOutputResource
+	 */
+	bool RequiresConversion(const TSharedPtr<IMediaTextureSample, ESPMode::ThreadSafe>& Sample, bool SrgbOutput)
 	{
-		return nullptr;
-	}
-
-	check(IsInRenderingThread());
-	check(SinkMode == EMediaTextureSinkMode::Unbuffered);
+		// If the output color space is expected to be sRGB, but the
+		// sample is not, a color space conversion on the GPU is required.
 
-	return RenderTargetTextureRHI;
-}
+		if (Sample->IsOutputSrgb() != SrgbOutput)
+		{
+			return true;
+		}
 
+		// If the output dimensions are not the same as the sample's
+		// dimensions, a resizing conversion on the GPU is required.
 
-void FMediaTextureResource::InitializeBuffer(FIntPoint OutputDim, FIntPoint BufferDim, EMediaTextureSinkFormat Format, EMediaTextureSinkMode Mode)
-{
-	State = EState::Initializing;
+		if (Sample->GetDim() != Sample->GetOutputDim())
+		{
+			return true;
+		}
 
-	RenderThreadTasks.Enqueue([=]() {
-		InitializeResource(OutputDim, BufferDim, Format, Mode);
-	});
-}
+		// Only the following pixel formats are supported natively.
+		// All other formats require a conversion on the GPU.
 
+		const EMediaTextureSampleFormat Format = Sample->GetFormat();
 
-void FMediaTextureResource::ReleaseBuffer()
-{
-	if (State != EState::Initialized)
-	{
-		return;
+		return ((Format != EMediaTextureSampleFormat::CharBGRA) &&
+				(Format != EMediaTextureSampleFormat::FloatRGB) &&
+				(Format != EMediaTextureSampleFormat::FloatRGBA));
 	}
 
-	if (SinkMode == EMediaTextureSinkMode::Buffered)
-	{
-		TripleBuffer.SwapWriteBuffers();
-	}
-	else
+
+	/**
+	 * Check whether the given sample requires an sRGB texture.
+	 *
+	 * @param Sample The sample to check.
+	 * @return true if an sRGB texture is required, false otherwise.
+	 */
+	bool RequiresSrgbTexture(const TSharedPtr<IMediaTextureSample, ESPMode::ThreadSafe>& Sample)
 	{
-		check(IsInRenderingThread());
+		if (!Sample->IsOutputSrgb())
+		{
+			return false;
+		}
 
-		TRefCountPtr<FRHITexture2D>& Texture = RequiresConversion
-			? BufferResources[2].RenderTarget
-			: RenderTargetTextureRHI;
+		const EMediaTextureSampleFormat Format = Sample->GetFormat();
 
-		RHIUnlockTexture2D(Texture.GetReference(), 0, false /*bLockWithinMiptail*/);
+		return ((Format == EMediaTextureSampleFormat::CharBGRA) ||
+				(Format == EMediaTextureSampleFormat::CharBMP) ||
+				(Format == EMediaTextureSampleFormat::FloatRGB) ||
+				(Format == EMediaTextureSampleFormat::FloatRGBA));
 	}
 }
 
 
-void FMediaTextureResource::ShutdownBuffer()
-{
-	if ((State == EState::ShutDown) || (State == EState::ShuttingDown))
-	{
-		return;
-	}
+/* FMediaTextureResource structors
+ *****************************************************************************/
 
-	State = EState::ShuttingDown;
+FMediaTextureResource::FMediaTextureResource(UMediaTexture& InOwner, FIntPoint& InOwnerDim, SIZE_T& InOwnerSize)
+	: Cleared(false)
+	, CurrentClearColor(FLinearColor::Transparent)
+	, Owner(InOwner)
+	, OwnerDim(InOwnerDim)
+	, OwnerSize(InOwnerSize)
+{ }
 
-	RenderThreadTasks.Enqueue([this]() {
-		ReleaseDynamicRHI();
-	});
-}
 
+/* FMediaTextureResource interface
+ *****************************************************************************/
 
-void FMediaTextureResource::UpdateBuffer(const uint8* Data, uint32 /*Pitch*/)
+void FMediaTextureResource::Render(const FRenderParams& Params)
 {
-	if (State != EState::Initialized)
-	{
-		return;
-	}
+	check(IsInRenderingThread());
 
-	// update sink
-	if (SinkMode == EMediaTextureSinkMode::Buffered)
-	{
-		uint8* Dest = (uint8*)TripleBuffer.GetWriteBuffer().LockedData;
+	TSharedPtr<FMediaTextureSampleSource, ESPMode::ThreadSafe> SampleSource = Params.SampleSource.Pin();
 
-		if (Dest != nullptr)
+	if (SampleSource.IsValid())
+	{
+		// get the most current sample to be rendered
+		TSharedPtr<IMediaTextureSample, ESPMode::ThreadSafe> Sample;
+		bool SampleValid = false;
+		
+		while (SampleSource->Peek(Sample) && Sample.IsValid())
 		{
-			FMemory::Memcpy(Dest, Data, BufferPitch * BufferDimensions.Y);
+			const FTimespan StartTime = Sample->GetTime();
+			const FTimespan EndTime = StartTime + Sample->GetDuration();
 
-			if (TripleBuffer.IsDirty())
+			if (((Params.Rate > 0.0f) && (StartTime > Params.Time)) ||
+				((Params.Rate < 0.0f) && (EndTime <= Params.Time)))
 			{
-				UE_LOG(LogMediaAssets, VeryVerbose, TEXT("MediaTextureResource frame dropped."));
+				break; // future sample
 			}
 
-			TripleBuffer.SwapWriteBuffers();
+			SampleValid = SampleSource->Dequeue(Sample);
 		}
-	}
-	else if (IsInRenderingThread())
-	{
-		if (RequiresConversion)
+
+		if (!SampleValid)
 		{
-			FUpdateTextureRegion2D Region(0, 0, 0, 0, BufferDimensions.X, BufferDimensions.Y);
-			RHIUpdateTexture2D(BufferResources[2].RenderTarget.GetReference(), 0, Region, BufferPitch, Data);
-			ConvertResource(BufferResources[2]);
+			return; // no sample to render
 		}
-		else
+
+		check(Sample.IsValid());
+
+		// render the sample
+		if (Sample->GetOutputDim().GetMin() <= 0)
 		{
-			FUpdateTextureRegion2D Region(0, 0, 0, 0, OutputDimensions.X, OutputDimensions.Y);
-			RHIUpdateTexture2D(RenderTargetTextureRHI.GetReference(), 0, Region, BufferPitch, Data);
+			#if MEDIATEXTURERESOURCE_TRACE_RENDER
+				UE_LOG(LogMediaAssets, VeryVerbose, TEXT("TextureResource %llx: Corrupt sample with time %s at time %s"), this, *Sample->GetTime().ToString(), *Params.Time.ToString());
+			#endif
+
+			ClearTexture(FLinearColor::Red, Params.SrgbOutput); // mark corrupt sample
 		}
-	}
-	else
-	{
-		UE_LOG(LogMediaAssets, Error, TEXT("UpdateTextureSinkBuffer must be called on render thread in Unbuffered mode."));
-	}
-}
+		else if (MediaTextureResource::RequiresConversion(Sample, Params.SrgbOutput))
+		{
+			#if MEDIATEXTURERESOURCE_TRACE_RENDER
+				UE_LOG(LogMediaAssets, VeryVerbose, TEXT("TextureResource %llx: Converting sample with time %s at time %s"), this, *Sample->GetTime().ToString(), *Params.Time.ToString());
+			#endif
 
+			ConvertSample(Sample, Params.ClearColor, Params.SrgbOutput);
+		}
+		else
+		{
+			#if MEDIATEXTURERESOURCE_TRACE_RENDER
+				UE_LOG(LogMediaAssets, VeryVerbose, TEXT("TextureResource %llx: Copying sample with time %s at time %s"), this, *Sample->GetTime().ToString(), *Params.Time.ToString());
+			#endif
 
-void FMediaTextureResource::UpdateTextures(FRHITexture* RenderTarget, FRHITexture* ShaderResource)
-{
-	if (State != EState::Initialized)
-	{
-		return;
+			CopySample(Sample, Params.ClearColor, Params.SrgbOutput);
+		}
+
+		if (!GSupportsImageExternal && Params.PlayerGuid.IsValid())
+		{
+			FTextureRHIRef VideoTexture = (FTextureRHIRef)Owner.TextureReference.TextureReferenceRHI;
+			FExternalTextureRegistry::Get().RegisterExternalTexture(Params.PlayerGuid, VideoTexture, SamplerStateRHI, Sample->GetScaleRotation(), Sample->GetOffset());
+		}
 	}
+	else if (!Cleared)
+	{
+		#if MEDIATEXTURERESOURCE_TRACE_RENDER
+			UE_LOG(LogMediaAssets, VeryVerbose, TEXT("TextureResource %llx: Clearing texture at time %s"), this, *Params.Time.ToString());
+		#endif
 
-	check(SinkMode == EMediaTextureSinkMode::Unbuffered);
+		ClearTexture(Params.ClearColor, Params.SrgbOutput);
 
-	FResource Resource;
-	{
-		Resource.RenderTarget = (RenderTarget != nullptr) ? RenderTarget->GetTexture2D() : nullptr;
-		Resource.ShaderResource = (ShaderResource != nullptr) ? ShaderResource->GetTexture2D() : nullptr;
+		if (!GSupportsImageExternal && Params.PlayerGuid.IsValid())
+		{
+			FTextureRHIRef VideoTexture = (FTextureRHIRef)Owner.TextureReference.TextureReferenceRHI;
+			FExternalTextureRegistry::Get().RegisterExternalTexture(Params.PlayerGuid, VideoTexture, SamplerStateRHI, FLinearColor(1.0f, 0.0f, 0.0f, 1.0f), FLinearColor(0.0f, 0.0f, 0.0f, 0.0f));
+		}
 	}
-
-	RenderThreadTasks.Enqueue([=]() {
-		SetResource(Resource);
-	});
 }
 
 
@@ -214,7 +227,7 @@ void FMediaTextureResource::UpdateTextures(FRHITexture* RenderTarget, FRHITextur
 
 FIntPoint FMediaTextureResource::GetSizeXY() const
 {
-	return OutputDimensions;
+	return FIntPoint(Owner.GetWidth(), Owner.GetHeight());
 }
 
 
@@ -229,13 +242,13 @@ FString FMediaTextureResource::GetFriendlyName() const
 
 uint32 FMediaTextureResource::GetSizeX() const
 {
-	return OutputDimensions.X;
+	return Owner.GetWidth();
 }
 
 
 uint32 FMediaTextureResource::GetSizeY() const
 {
-	return OutputDimensions.Y;
+	return Owner.GetHeight();
 }
 
 
@@ -250,563 +263,388 @@ void FMediaTextureResource::InitDynamicRHI()
 	);
 
 	SamplerStateRHI = RHICreateSamplerState(SamplerStateInitializer);
+}
 
-	if ((OutputDimensions.GetMin() <= 0) || (BufferDimensions.GetMin() <= 0))
-	{
-		return;
-	}
-
-	// determine buffer resource pixel format
-	EPixelFormat BufferFormat;
-	bool NeedsResize;
-	bool NeedsTransform;
-	bool SrgbBuffer;
-
-	switch (SinkFormat)
-	{
-	case EMediaTextureSinkFormat::CharAYUV:
-		BufferFormat = PF_B8G8R8A8;
-		BufferPitch = BufferDimensions.X * 4;
-		NeedsResize = false;
-		NeedsTransform = true;
-		SrgbBuffer = false;
-		break;
-
-	case EMediaTextureSinkFormat::CharBGRA:
-		BufferFormat = PF_B8G8R8A8;
-		BufferPitch = BufferDimensions.X * 4;
-		NeedsResize = (BufferDimensions != OutputDimensions);
-		NeedsTransform = false;
-		SrgbBuffer = true;
-		break;
-
-	case EMediaTextureSinkFormat::CharBMP:
-		BufferFormat = PF_B8G8R8A8;
-		BufferPitch = BufferDimensions.X * 4;
-		NeedsResize = true;
-		NeedsTransform = false;
-		SrgbBuffer = true;
-		break;
-
-	case EMediaTextureSinkFormat::CharNV12:
-	case EMediaTextureSinkFormat::CharNV21:
-		BufferFormat = PF_G8;
-		BufferPitch = BufferDimensions.X;
-		NeedsResize = false;
-		NeedsTransform = true;
-		SrgbBuffer = false;
-		break;
-
-	case EMediaTextureSinkFormat::CharUYVY:
-	case EMediaTextureSinkFormat::CharYUY2:
-	case EMediaTextureSinkFormat::CharYVYU:
-		BufferFormat = PF_B8G8R8A8;
-		BufferPitch = BufferDimensions.X * 4;
-		NeedsResize = false;
-		NeedsTransform = true;
-		SrgbBuffer = false;
-		break;
-
-	case EMediaTextureSinkFormat::FloatRGB:
-		BufferFormat = PF_FloatRGB;
-		BufferPitch = BufferDimensions.X * 12;
-		NeedsResize = (BufferDimensions != OutputDimensions);
-		NeedsTransform = false;
-		SrgbBuffer = true;
-		break;
-
-	case EMediaTextureSinkFormat::FloatRGBA:
-		BufferFormat = PF_FloatRGBA;
-		BufferPitch = BufferDimensions.X * 16;
-		NeedsResize = (BufferDimensions != OutputDimensions);
-		NeedsTransform = false;
-		SrgbBuffer = true;
-		break;
-
-	default:
-		return; // unsupported sink format
-	}
-
-	RequiresConversion = NeedsResize || NeedsTransform;
 
-	// determine how many buffer resources we need
-	int32 StartIndex;
-	
-	if (SinkMode == EMediaTextureSinkMode::Buffered)
-	{
-		StartIndex = 0; // triple buffer
-	}
-	else if (RequiresConversion)
-	{
-		StartIndex = 2; // single buffer (frame before pixel format conversion)
-	}
-	else
-	{
-		StartIndex = 3; // no buffering (directly writing to output or replacing resources)
-	}
+void FMediaTextureResource::ReleaseDynamicRHI()
+{
+	Cleared = false;
 
-	// create buffer resources
-	if (StartIndex < 3)
-	{
-		TripleBuffer.Reset();
+	InputTarget.SafeRelease();
+	OutputTarget.SafeRelease();
+	RenderTargetTextureRHI.SafeRelease();
+	TextureRHI.SafeRelease();
 
-		FRHIResourceCreateInfo CreateInfo = { FClearValueBinding(BufferClearColor) };
-		uint32 OutStride = 0;
+	UpdateTextureReference(nullptr);
+}
 
-		// only one mip is used & textures are targetable or resolve
-		for (int32 ResourceIndex = StartIndex; ResourceIndex < 3; ++ResourceIndex)
-		{
-			FResource& Resource = BufferResources[ResourceIndex];
 
-			RHICreateTargetableShaderResource2D(
-				BufferDimensions.X,
-				BufferDimensions.Y,
-				BufferFormat,
-				1,
-				TexCreate_Dynamic | (SrgbBuffer ? TexCreate_SRGB : 0),
-				TexCreate_RenderTargetable,
-				false,
-				CreateInfo,
-				Resource.RenderTarget,
-				Resource.ShaderResource
-			);
+/* FMediaTextureResource implementation
+ *****************************************************************************/
 
-			if (ResourceIndex < 2)
-			{
-				// lock write & temp buffers (in triple buffered mode only)
-				Resource.LockedData = RHILockTexture2D(Resource.RenderTarget.GetReference(), 0, RLM_WriteOnly, OutStride, false);
-			}
-			else
-			{
-				Resource.LockedData = nullptr;
-			}
-		}
-	}
+void FMediaTextureResource::ClearTexture(const FLinearColor& ClearColor, bool SrgbOutput)
+{
+	// create output render target if we don't have one yet
+	const uint32 OutputCreateFlags = TexCreate_Dynamic | (SrgbOutput ? TexCreate_SRGB : 0);
 
-	// create output resource if needed
-	if (RequiresConversion || (SinkMode == EMediaTextureSinkMode::Unbuffered))
+	if ((ClearColor != CurrentClearColor) || !OutputTarget.IsValid() || ((OutputTarget->GetFlags() & OutputCreateFlags) != OutputCreateFlags))
 	{
-		// In Unbuffered mode or in Buffered mode (with conversion), this
-		// resource is used as a render target for the output frame.
+		FRHIResourceCreateInfo CreateInfo = {
+			FClearValueBinding(ClearColor)
+		};
 
-		FRHIResourceCreateInfo CreateInfo = { FClearValueBinding(BufferClearColor) };
-		TRefCountPtr<FRHITexture2D> Texture2DRHI;
+		TRefCountPtr<FRHITexture2D> DummyTexture2DRHI;
 
 		RHICreateTargetableShaderResource2D(
-			OutputDimensions.X,
-			OutputDimensions.Y,
-			RequiresConversion ? PF_B8G8R8A8 : BufferFormat,
+			2,
+			2,
+			PF_B8G8R8A8,
 			1,
-			TexCreate_Dynamic | (Owner.SRGB ? TexCreate_SRGB : 0),
+			OutputCreateFlags,
 			TexCreate_RenderTargetable,
 			false,
 			CreateInfo,
-			RenderTargetTextureRHI,
-			Texture2DRHI
+			OutputTarget,
+			DummyTexture2DRHI
 		);
 
-		TextureRHI = (FTextureRHIRef&)Texture2DRHI;
+		CurrentClearColor = ClearColor;
+		UpdateResourceSize();
 	}
-	else
-	{
-		check(StartIndex < 3);
 
-		// in Buffered mode, we pick the current read buffer as the first resource
-		// to be displayed. In Unbuffered mode (with conversion), the read buffer
-		// is used as the output render target.
-
-		const FResource& Resource = TripleBuffer.Read();
-		
-		RenderTargetTextureRHI = Resource.RenderTarget;
-		TextureRHI = Resource.ShaderResource;
+	if (RenderTargetTextureRHI != OutputTarget)
+	{
+		UpdateTextureReference(OutputTarget);
 	}
 
-	// set the initial resource to be displayed
-	RHIUpdateTextureReference(Owner.TextureReference.TextureReferenceRHI, RenderTargetTextureRHI.GetReference());
-
-	// clear the displayed resource
-	FRHICommandListImmediate& RHICmdList = FRHICommandListExecutor::GetImmediateCommandList();
+	// draw the clear color
+	FRHICommandListImmediate& CommandList = FRHICommandListExecutor::GetImmediateCommandList();
 	{
 		FRHIRenderTargetView View = FRHIRenderTargetView(RenderTargetTextureRHI, ERenderTargetLoadAction::EClear);
 		FRHISetRenderTargetsInfo Info(1, &View, FRHIDepthRenderTargetView());
-		RHICmdList.SetRenderTargetsAndClear(Info);
-		RHICmdList.SetViewport(0, 0, 0.0f, OutputDimensions.X, OutputDimensions.Y, 1.0f);
-		RHICmdList.TransitionResource(EResourceTransitionAccess::EReadable, TextureRHI);
+		CommandList.SetRenderTargetsAndClear(Info);
+		CommandList.SetViewport(0, 0, 0.0f, RenderTargetTextureRHI->GetSizeX(), RenderTargetTextureRHI->GetSizeY(), 1.0f);
+		CommandList.TransitionResource(EResourceTransitionAccess::EReadable, RenderTargetTextureRHI);
 	}
 
-	CacheResourceSize();
-
-	State = EState::Initialized;
-	FPlatformMisc::MemoryBarrier();
+	Cleared = true;
 }
 
 
-void FMediaTextureResource::ReleaseDynamicRHI()
+void FMediaTextureResource::ConvertSample(const TSharedPtr<IMediaTextureSample, ESPMode::ThreadSafe>& Sample, const FLinearColor& ClearColor, bool SrgbOutput)
 {
-	// release output resource if needed
-	if (RequiresConversion || (SinkMode == EMediaTextureSinkMode::Unbuffered))
-	{
-		RenderTargetTextureRHI.SafeRelease();
-		TextureRHI.SafeRelease();
-	}
-
-	// release buffer resources
-	int32 StartIndex;
+	const EPixelFormat PixelFormat = MediaTextureResource::GetPixelFormat(Sample);
 
-	if (SinkMode == EMediaTextureSinkMode::Buffered)
-	{
-		StartIndex = 0; // triple buffer
-	}
-	else if (RequiresConversion)
-	{
-		StartIndex = 2; // single buffer (for pixel format conversion)
-	}
-	else
+	// get input texture
+	FRHITexture2D* InputTexture = nullptr;
 	{
-		StartIndex = 3; // no buffering (writing directly to output)
-	}
+		// If the sample already provides a texture resource, we simply use that
+		// as the input texture. If the sample only provides raw data, then we
+		// create our own input render target and copy the data into it.
 
-	for (int32 ResourceIndex = StartIndex; ResourceIndex < 3; ++ResourceIndex)
-	{
-		FResource& Resource = BufferResources[ResourceIndex];
+		FRHITexture* SampleTexture = Sample->GetTexture();
+		FRHITexture2D* SampleTexture2D = (SampleTexture != nullptr) ? SampleTexture->GetTexture2D() : nullptr;
 
-		if (Resource.LockedData != nullptr)
+		if (SampleTexture2D != nullptr)
 		{
-			RHIUnlockTexture2D(Resource.RenderTarget.GetReference(), 0, false);
-			Resource.LockedData = nullptr;
-		}
+			InputTexture = SampleTexture2D;
 
-		Resource.ShaderResource.SafeRelease();
-		Resource.RenderTarget.SafeRelease();
-	}
-
-	// clean up
-	if (SinkMode == EMediaTextureSinkMode::Buffered)
-	{
-		TripleBuffer.Reset();
-	}
+			InputTarget.SafeRelease();
+			UpdateResourceSize();
+		}
+		else
+		{
+			const bool SrgbTexture = MediaTextureResource::RequiresSrgbTexture(Sample);
+			const uint32 InputCreateFlags = TexCreate_Dynamic | (SrgbTexture ? TexCreate_SRGB : 0);
+			const FIntPoint SampleDim = Sample->GetDim();
 
-	CacheResourceSize();
+			// create a new input render target if necessary
+			if (!InputTarget.IsValid() || (InputTarget->GetSizeXY() != SampleDim) || (InputTarget->GetFormat() != PixelFormat) || ((InputTarget->GetFlags() & InputCreateFlags) != InputCreateFlags))
+			{
+				TRefCountPtr<FRHITexture2D> DummyTexture2DRHI;
+				FRHIResourceCreateInfo CreateInfo;
+
+				RHICreateTargetableShaderResource2D(
+					SampleDim.X,
+					SampleDim.Y,
+					PixelFormat,
+					1,
+					InputCreateFlags,
+					TexCreate_RenderTargetable,
+					false,
+					CreateInfo,
+					InputTarget,
+					DummyTexture2DRHI
+				);
+
+				UpdateResourceSize();
+			}
 
-	State = EState::ShutDown;
-	FPlatformMisc::MemoryBarrier();
-}
+			// copy sample data to input render target
+			FUpdateTextureRegion2D Region(0, 0, 0, 0, SampleDim.X, SampleDim.Y);
+			RHIUpdateTexture2D(InputTarget, 0, Region, Sample->GetStride(), (uint8*)Sample->GetBuffer());
 
+			InputTexture = InputTarget;
+		}
+	}
 
-/* FTickerObjectBase interface
- *****************************************************************************/
+	// create output render target if necessary
+	const uint32 OutputCreateFlags = TexCreate_Dynamic | (SrgbOutput ? TexCreate_SRGB : 0);
+	const FIntPoint OutputDim = Sample->GetOutputDim();
 
-bool FMediaTextureResource::Tick(float DeltaTime)
-{
-	ENQUEUE_RENDER_COMMAND(TickMediaTextureResource)([this](FRHICommandListImmediate&)
+	if ((ClearColor != CurrentClearColor) || !OutputTarget.IsValid() || (OutputTarget->GetSizeXY() != OutputDim) || (OutputTarget->GetFormat() != PF_B8G8R8A8) || ((OutputTarget->GetFlags() & OutputCreateFlags) != OutputCreateFlags))
 	{
-		ProcessRenderThreadTasks();
-	});
-
-	return true;
-}
-
+		TRefCountPtr<FRHITexture2D> DummyTexture2DRHI;
+		
+		FRHIResourceCreateInfo CreateInfo = {
+			FClearValueBinding(ClearColor)
+		};
 
-/* FMediaTextureResource implementation
- *****************************************************************************/
+		RHICreateTargetableShaderResource2D(
+			OutputDim.X,
+			OutputDim.Y,
+			PF_B8G8R8A8,
+			1,
+			OutputCreateFlags,
+			TexCreate_RenderTargetable,
+			false,
+			CreateInfo,
+			OutputTarget,
+			DummyTexture2DRHI
+		);
 
-void FMediaTextureResource::CacheResourceSize()
-{
-	CachedResourceSizeBytes = 0;
+		CurrentClearColor = ClearColor;
+		UpdateResourceSize();
+	}
 
-	if (OutputDimensions.GetMin() > 0)
+	if (RenderTargetTextureRHI != OutputTarget)
 	{
-		const SIZE_T BufferSize = BufferPitch * BufferDimensions.Y;
-
-		if (SinkMode == EMediaTextureSinkMode::Buffered)
-		{
-			CachedResourceSizeBytes += (3 * BufferSize);
-		}
-		else
-		{
-			CachedResourceSizeBytes += BufferSize;
-		}
-
-		if (RequiresConversion)
-		{
-			CachedResourceSizeBytes += (OutputDimensions.X * OutputDimensions.Y * 4);
-		}
+		UpdateTextureReference(OutputTarget);
 	}
-}
-
-
-void FMediaTextureResource::ConvertResource(const FResource& Resource)
-{
-	check(IsInRenderingThread());
-
-	FRHICommandListImmediate& RHICmdList = FRHICommandListExecutor::GetImmediateCommandList();
 
-	SCOPED_DRAW_EVENT(RHICmdList, MediaTextureConvertResource);
+	// perform the conversion
+	FRHICommandListImmediate& CommandList = FRHICommandListExecutor::GetImmediateCommandList();
+	{
+		SCOPED_DRAW_EVENT(CommandList, MediaTextureConvertResource);
 
-	FGraphicsPipelineStateInitializer GraphicsPSOInit;
-	FRHITexture* RenderTarget = RenderTargetTextureRHI.GetReference();
-	SetRenderTargets(RHICmdList, 1, &RenderTarget, nullptr, ESimpleRenderTargetMode::EExistingColorAndDepth, FExclusiveDepthStencil::DepthNop_StencilNop);
-	RHICmdList.ApplyCachedRenderTargets(GraphicsPSOInit);
+		FGraphicsPipelineStateInitializer GraphicsPSOInit;
+		FRHITexture* RenderTarget = RenderTargetTextureRHI.GetReference();
+		SetRenderTargets(CommandList, 1, &RenderTarget, nullptr, ESimpleRenderTargetMode::EExistingColorAndDepth, FExclusiveDepthStencil::DepthNop_StencilNop);
 
-	RHICmdList.SetViewport(0, 0, 0.0f, OutputDimensions.X, OutputDimensions.Y, 1.0f);
+		CommandList.ApplyCachedRenderTargets(GraphicsPSOInit);
+		CommandList.SetViewport(0, 0, 0.0f, OutputDim.X, OutputDim.Y, 1.0f);
 
-	GraphicsPSOInit.DepthStencilState = TStaticDepthStencilState<false, CF_Always>::GetRHI();
-	GraphicsPSOInit.RasterizerState = TStaticRasterizerState<>::GetRHI();
-	GraphicsPSOInit.BlendState = TStaticBlendStateWriteMask<CW_RGBA, CW_NONE, CW_NONE, CW_NONE, CW_NONE, CW_NONE, CW_NONE, CW_NONE>::GetRHI();
-	GraphicsPSOInit.PrimitiveType = PT_TriangleStrip;
+		GraphicsPSOInit.DepthStencilState = TStaticDepthStencilState<false, CF_Always>::GetRHI();
+		GraphicsPSOInit.RasterizerState = TStaticRasterizerState<>::GetRHI();
+		GraphicsPSOInit.BlendState = TStaticBlendStateWriteMask<CW_RGBA, CW_NONE, CW_NONE, CW_NONE, CW_NONE, CW_NONE, CW_NONE, CW_NONE>::GetRHI();
+		GraphicsPSOInit.PrimitiveType = PT_TriangleStrip;
 
-	// configure media shaders
-	auto ShaderMap = GetGlobalShaderMap(GMaxRHIFeatureLevel);
-	TShaderMapRef<FMediaShadersVS> VertexShader(ShaderMap);
+		// configure media shaders
+		auto ShaderMap = GetGlobalShaderMap(GMaxRHIFeatureLevel);
+		TShaderMapRef<FMediaShadersVS> VertexShader(ShaderMap);
 
-	GraphicsPSOInit.BoundShaderState.VertexDeclarationRHI = GMediaVertexDeclaration.VertexDeclarationRHI;
-	GraphicsPSOInit.BoundShaderState.VertexShaderRHI = GETSAFERHISHADER_VERTEX(*VertexShader);
+		GraphicsPSOInit.BoundShaderState.VertexDeclarationRHI = GMediaVertexDeclaration.VertexDeclarationRHI;
+		GraphicsPSOInit.BoundShaderState.VertexShaderRHI = GETSAFERHISHADER_VERTEX(*VertexShader);
 
-	switch (SinkFormat)
-	{
-	case EMediaTextureSinkFormat::CharAYUV:
+		switch (Sample->GetFormat())
+		{
+		case EMediaTextureSampleFormat::CharAYUV:
 		{
 			TShaderMapRef<FAYUVConvertPS> ConvertShader(ShaderMap);
-
 			GraphicsPSOInit.BoundShaderState.PixelShaderRHI = GETSAFERHISHADER_PIXEL(*ConvertShader);
-			SetGraphicsPipelineState(RHICmdList, GraphicsPSOInit);
+			SetGraphicsPipelineState(CommandList, GraphicsPSOInit);
+			ConvertShader->SetParameters(CommandList, InputTexture, MediaShaders::YuvToSrgbDefault, Sample->IsOutputSrgb());
+		}
+		break;
 
-			ConvertShader->SetParameters(
-				RHICmdList,
-				Resource.ShaderResource,
-				MediaShaders::YuvToSrgbDefault,
-				Owner.SRGB
-			);
+		case EMediaTextureSampleFormat::CharBMP:
+		{
+			TShaderMapRef<FBMPConvertPS> ConvertShader(ShaderMap);
+			GraphicsPSOInit.BoundShaderState.PixelShaderRHI = GETSAFERHISHADER_PIXEL(*ConvertShader);
+			SetGraphicsPipelineState(CommandList, GraphicsPSOInit);
+			ConvertShader->SetParameters(CommandList, InputTexture, OutputDim, Sample->IsOutputSrgb() && !SrgbOutput);
 		}
 		break;
 
-	case EMediaTextureSinkFormat::CharNV12:
+		case EMediaTextureSampleFormat::CharNV12:
 		{
 			TShaderMapRef<FNV12ConvertPS> ConvertShader(ShaderMap);
-
 			GraphicsPSOInit.BoundShaderState.PixelShaderRHI = GETSAFERHISHADER_PIXEL(*ConvertShader);
-			SetGraphicsPipelineState(RHICmdList, GraphicsPSOInit);
-
-			ConvertShader->SetParameters(
-				RHICmdList,
-				Resource.ShaderResource,
-				OutputDimensions,
-				MediaShaders::YuvToSrgbDefault,
-				Owner.SRGB
-			);
+			SetGraphicsPipelineState(CommandList, GraphicsPSOInit);
+			ConvertShader->SetParameters(CommandList, InputTexture, OutputDim, MediaShaders::YuvToSrgbDefault, Sample->IsOutputSrgb());
 		}
 		break;
 
-	case EMediaTextureSinkFormat::CharNV21:
+		case EMediaTextureSampleFormat::CharNV21:
 		{
 			TShaderMapRef<FNV21ConvertPS> ConvertShader(ShaderMap);
-
-			GraphicsPSOInit.BoundShaderState.PixelShaderRHI = GETSAFERHISHADER_PIXEL(*ConvertShader); 
-			SetGraphicsPipelineState(RHICmdList, GraphicsPSOInit);
-
-			ConvertShader->SetParameters(
-				RHICmdList,
-				Resource.ShaderResource,
-				OutputDimensions,
-				MediaShaders::YuvToSrgbDefault,
-				Owner.SRGB
-			);
+			GraphicsPSOInit.BoundShaderState.PixelShaderRHI = GETSAFERHISHADER_PIXEL(*ConvertShader);
+			SetGraphicsPipelineState(CommandList, GraphicsPSOInit);
+			ConvertShader->SetParameters(CommandList, InputTexture, OutputDim, MediaShaders::YuvToSrgbDefault, Sample->IsOutputSrgb());
 		}
 		break;
 
-	case EMediaTextureSinkFormat::CharUYVY:
+		case EMediaTextureSampleFormat::CharUYVY:
 		{
 			TShaderMapRef<FUYVYConvertPS> ConvertShader(ShaderMap);
-
 			GraphicsPSOInit.BoundShaderState.PixelShaderRHI = GETSAFERHISHADER_PIXEL(*ConvertShader);
-			SetGraphicsPipelineState(RHICmdList, GraphicsPSOInit);
-
-			ConvertShader->SetParameters(
-				RHICmdList,
-				Resource.ShaderResource,
-				MediaShaders::YuvToSrgbDefault,
-				Owner.SRGB
-			);
+			SetGraphicsPipelineState(CommandList, GraphicsPSOInit);
+			ConvertShader->SetParameters(CommandList, InputTexture, MediaShaders::YuvToSrgbDefault, Sample->IsOutputSrgb());
 		}
 		break;
 
-	case EMediaTextureSinkFormat::CharYUY2:
+		case EMediaTextureSampleFormat::CharYUY2:
 		{
 			TShaderMapRef<FYUY2ConvertPS> ConvertShader(ShaderMap);
-
 			GraphicsPSOInit.BoundShaderState.PixelShaderRHI = GETSAFERHISHADER_PIXEL(*ConvertShader);
-			SetGraphicsPipelineState(RHICmdList, GraphicsPSOInit);
-
-			ConvertShader->SetParameters(
-				RHICmdList,
-				Resource.ShaderResource,
-				OutputDimensions,
-				MediaShaders::YuvToSrgbDefault,
-				Owner.SRGB
-			);
+			SetGraphicsPipelineState(CommandList, GraphicsPSOInit);
+			ConvertShader->SetParameters(CommandList, InputTexture, OutputDim, MediaShaders::YuvToSrgbDefault, Sample->IsOutputSrgb());
 		}
 		break;
 
-	case EMediaTextureSinkFormat::CharYVYU:
+		case EMediaTextureSampleFormat::CharYVYU:
 		{
 			TShaderMapRef<FYVYUConvertPS> ConvertShader(ShaderMap);
-
 			GraphicsPSOInit.BoundShaderState.PixelShaderRHI = GETSAFERHISHADER_PIXEL(*ConvertShader);
-			SetGraphicsPipelineState(RHICmdList, GraphicsPSOInit);
-
-			ConvertShader->SetParameters(
-				RHICmdList,
-				Resource.ShaderResource,
-				MediaShaders::YuvToSrgbDefault,
-				Owner.SRGB
-			);
+			SetGraphicsPipelineState(CommandList, GraphicsPSOInit);
+			ConvertShader->SetParameters(CommandList, InputTexture, MediaShaders::YuvToSrgbDefault, Sample->IsOutputSrgb());
 		}
 		break;
 
-	case EMediaTextureSinkFormat::CharBGRA:
-	case EMediaTextureSinkFormat::FloatRGB:
-	case EMediaTextureSinkFormat::FloatRGBA:
+		case EMediaTextureSampleFormat::CharBGRA:
+		case EMediaTextureSampleFormat::FloatRGB:
+		case EMediaTextureSampleFormat::FloatRGBA:
 		{
 			TShaderMapRef<FRGBConvertPS> ConvertShader(ShaderMap);
-
 			GraphicsPSOInit.BoundShaderState.PixelShaderRHI = GETSAFERHISHADER_PIXEL(*ConvertShader);
-			SetGraphicsPipelineState(RHICmdList, GraphicsPSOInit);
-			ConvertShader->SetParameters(RHICmdList, Resource.ShaderResource, OutputDimensions, Owner.SRGB);
+			SetGraphicsPipelineState(CommandList, GraphicsPSOInit);
+			ConvertShader->SetParameters(CommandList, InputTexture, OutputDim);
 		}
 		break;
 
-	case EMediaTextureSinkFormat::CharBMP:
-		{
-			TShaderMapRef<FBMPConvertPS> ConvertShader(ShaderMap);
-
-			GraphicsPSOInit.BoundShaderState.PixelShaderRHI = GETSAFERHISHADER_PIXEL(*ConvertShader);
-			SetGraphicsPipelineState(RHICmdList, GraphicsPSOInit);
-			ConvertShader->SetParameters(RHICmdList, Resource.ShaderResource, OutputDimensions);
+		default:
+			return; // unsupported format
 		}
-		break;
 
-	default:
-		return;
-	}
+		// draw full size quad into render target
+		FMediaElementVertex Vertices[4];
+		{
+			Vertices[0].Position.Set(-1.0f, 1.0f, 1.0f, 1.0f);
+			Vertices[0].TextureCoordinate.Set(0.0f, 0.0f);
+			Vertices[1].Position.Set(1.0f, 1.0f, 1.0f, 1.0f);
+			Vertices[1].TextureCoordinate.Set(1.0f, 0.0f);
+			Vertices[2].Position.Set(-1.0f, -1.0f, 1.0f, 1.0f);
+			Vertices[2].TextureCoordinate.Set(0.0f, 1.0f);
+			Vertices[3].Position.Set(1.0f, -1.0f, 1.0f, 1.0f);
+			Vertices[3].TextureCoordinate.Set(1.0f, 1.0f);
+		}
 
-	// draw full size quad into render target
-	FMediaElementVertex Vertices[4];
-	{
-		Vertices[0].Position.Set(-1.0f, 1.0f, 1.0f, 1.0f);
-		Vertices[0].TextureCoordinate.Set(0.0f, 0.0f);
-		Vertices[1].Position.Set(1.0f, 1.0f, 1.0f, 1.0f);
-		Vertices[1].TextureCoordinate.Set(1.0f, 0.0f);
-		Vertices[2].Position.Set(-1.0f, -1.0f, 1.0f, 1.0f);
-		Vertices[2].TextureCoordinate.Set(0.0f, 1.0f);
-		Vertices[3].Position.Set(1.0f, -1.0f, 1.0f, 1.0f);
-		Vertices[3].TextureCoordinate.Set(1.0f, 1.0f);
+		CommandList.SetViewport(0, 0, 0.0f, OutputDim.X, OutputDim.Y, 1.0f);
+		DrawPrimitiveUP(CommandList, PT_TriangleStrip, 2, Vertices, sizeof(Vertices[0]));
+		CommandList.TransitionResource(EResourceTransitionAccess::EReadable, RenderTargetTextureRHI);
 	}
 
-	RHICmdList.SetViewport(0, 0, 0.0f, OutputDimensions.X, OutputDimensions.Y, 1.0f);
-
-	DrawPrimitiveUP(RHICmdList, PT_TriangleStrip, 2, Vertices, sizeof(Vertices[0]));
-
-	RHICmdList.TransitionResource(EResourceTransitionAccess::EReadable, RenderTargetTextureRHI);
+	Cleared = false;
 }
 
 
-void FMediaTextureResource::DisplayResource(const FResource& Resource)
+void FMediaTextureResource::CopySample(const TSharedPtr<IMediaTextureSample, ESPMode::ThreadSafe>& Sample, const FLinearColor& ClearColor, bool SrgbOutput)
 {
-	check(IsInRenderingThread());
+	FRHITexture* SampleTexture = Sample->GetTexture();
+	FRHITexture2D* SampleTexture2D = (SampleTexture != nullptr) ? SampleTexture->GetTexture2D() : nullptr;
 
-	RenderTargetTextureRHI = Resource.RenderTarget;
-	TextureRHI = Resource.ShaderResource;
-
-	RHIUpdateTextureReference(Owner.TextureReference.TextureReferenceRHI, RenderTargetTextureRHI.GetReference());
-}
-
-
-void FMediaTextureResource::InitializeResource(FIntPoint OutputDim, FIntPoint BufferDim, EMediaTextureSinkFormat Format, EMediaTextureSinkMode Mode)
-{
-	check(IsInRenderingThread());
+	// If the sample already provides a texture resource, we simply use that
+	// as the output render target. If the sample only provides raw data, then
+	// we create our own output render target and copy the data into it.
 
-	if ((OutputDim == OutputDimensions) &&
-		(BufferDim == BufferDimensions) &&
-		(Format == SinkFormat) &&
-		(Mode == SinkMode))
+	if (SampleTexture2D != nullptr)
 	{
-		// reuse existing resources
-		State = (OutputDimensions.GetMin() > 0) ? EState::Initialized : EState::ShutDown;
-		FPlatformMisc::MemoryBarrier();
+		// use sample's texture as the new render target.
+		if (TextureRHI != SampleTexture2D)
+		{
+			UpdateTextureReference(SampleTexture2D);
+
+			OutputTarget.SafeRelease();
+			UpdateResourceSize();
+		}
 	}
 	else
 	{
-		// reinitialize resources
-		ReleaseDynamicRHI();
-		BufferDimensions = BufferDim;
-		OutputDimensions = OutputDim;
-		SinkFormat = Format;
-		SinkMode = Mode;
-		InitDynamicRHI();
-	}
-}
+		// create a new output render target if necessary
+		const uint32 OutputCreateFlags = TexCreate_Dynamic | (SrgbOutput ? TexCreate_SRGB : 0);
+		const EPixelFormat SampleFormat = MediaTextureResource::GetPixelFormat(Sample);
+		const FIntPoint SampleDim = Sample->GetDim();
 
+		if ((ClearColor != CurrentClearColor) || !OutputTarget.IsValid() || (OutputTarget->GetSizeXY() != SampleDim) || (OutputTarget->GetFormat() != SampleFormat) || ((OutputTarget->GetFlags() & OutputCreateFlags) != OutputCreateFlags))
+		{
+			TRefCountPtr<FRHITexture2D> DummyTexture2DRHI;
 
-void FMediaTextureResource::ProcessRenderThreadTasks()
-{
-	TFunction<void()> Task;
+			FRHIResourceCreateInfo CreateInfo = {
+				FClearValueBinding(ClearColor)
+			};
 
-	while (RenderThreadTasks.Dequeue(Task))
-	{
-		Task();
+			RHICreateTargetableShaderResource2D(
+				SampleDim.X,
+				SampleDim.Y,
+				SampleFormat,
+				1,
+				OutputCreateFlags,
+				TexCreate_RenderTargetable,
+				false,
+				CreateInfo,
+				OutputTarget,
+				DummyTexture2DRHI
+			);
+
+			CurrentClearColor = ClearColor;
+			UpdateResourceSize();
+		}
+
+		if (RenderTargetTextureRHI != OutputTarget)
+		{
+			UpdateTextureReference(OutputTarget);
+		}
+
+		// copy sample data to output render target
+		FUpdateTextureRegion2D Region(0, 0, 0, 0, SampleDim.X, SampleDim.Y);
+		RHIUpdateTexture2D(RenderTargetTextureRHI.GetReference(), 0, Region, Sample->GetStride(), (uint8*)Sample->GetBuffer());
 	}
+
+	Cleared = false;
 }
 
 
-void FMediaTextureResource::SetResource(const FResource& Resource)
+void FMediaTextureResource::UpdateResourceSize()
 {
-	check(IsInRenderingThread());
+	SIZE_T ResourceSize = 0;
 
-	if (!Resource.RenderTarget.IsValid() || !Resource.ShaderResource.IsValid())
+	if (InputTarget.IsValid())
 	{
-		ReleaseDynamicRHI();
-		InitDynamicRHI();
+		ResourceSize += CalcTextureSize(InputTarget->GetSizeX(), InputTarget->GetSizeY(), InputTarget->GetFormat(), 1);
 	}
-	else if (RequiresConversion)
-	{
-		ConvertResource(Resource);
-	}
-	else
+
+	if (OutputTarget.IsValid())
 	{
-		DisplayResource(Resource);
+		ResourceSize += CalcTextureSize(OutputTarget->GetSizeX(), OutputTarget->GetSizeY(), OutputTarget->GetFormat(), 1);
 	}
+
+	OwnerSize = ResourceSize;
 }
 
 
-void FMediaTextureResource::SwapResource()
+void FMediaTextureResource::UpdateTextureReference(FRHITexture2D* NewTexture)
 {
-	if (SinkMode == EMediaTextureSinkMode::Buffered)
+	TextureRHI = NewTexture;
+	RenderTargetTextureRHI = NewTexture;
+
+	RHIUpdateTextureReference(Owner.TextureReference.TextureReferenceRHI, NewTexture);
+
+	if (RenderTargetTextureRHI != nullptr)
 	{
-		if (TripleBuffer.IsDirty())
-		{
-			// lock & swap out old read buffer
-			FResource& OldResource = TripleBuffer.Read();
-			uint32 OutStride = 0;
-			OldResource.LockedData = RHILockTexture2D(OldResource.RenderTarget.GetReference(), 0, RLM_WriteOnly, OutStride, false /*bLockWithinMiptail*/, false /*bFlushRHIThread*/);
-			TripleBuffer.SwapReadBuffers();
-
-			// unlock & assign new read buffer
-			FResource& NewResource = TripleBuffer.Read();
-			RHIUnlockTexture2D(NewResource.RenderTarget.GetReference(), 0, false /*bLockWithinMiptail*/);
-			NewResource.LockedData = nullptr;
-
-			if (RequiresConversion)
-			{
-				ConvertResource(NewResource);
-			}
-			else
-			{
-				DisplayResource(NewResource);
-			}
-		}
+		OwnerDim = FIntPoint(RenderTargetTextureRHI->GetSizeX(), RenderTargetTextureRHI->GetSizeY());
 	}
-	else if (RequiresConversion)
+	else
 	{
-		ConvertResource(BufferResources[2]);
+		OwnerDim = FIntPoint::ZeroValue;
 	}
 }
diff --git a/Engine/Source/Runtime/MediaAssets/Private/Misc/MediaTextureResource.h b/Engine/Source/Runtime/MediaAssets/Private/Misc/MediaTextureResource.h
--- a/Engine/Source/Runtime/MediaAssets/Private/Misc/MediaTextureResource.h
+++ b/Engine/Source/Runtime/MediaAssets/Private/Misc/MediaTextureResource.h
@@ -2,14 +2,20 @@
 
 #pragma once
 
-#include "CoreMinimal.h"
-#include "Containers/Ticker.h"
-#include "Containers/TripleBuffer.h"
-#include "UnrealClient.h"
-#include "TextureResource.h"
+#include "CoreTypes.h"
 #include "Containers/Queue.h"
+#include "Math/Color.h"
+#include "MediaSampleSource.h"
+#include "Misc/Guid.h"
+#include "Misc/Timespan.h"
+#include "Templates/RefCounting.h"
+#include "Templates/SharedPointer.h"
+#include "TextureResource.h"
+#include "UnrealClient.h"
 
-
+class FMediaPlayerFacade;
+class IMediaPlayer;
+class IMediaTextureSample;
 class UMediaTexture;
 
 enum class EMediaTextureSinkFormat;
@@ -22,104 +28,53 @@ enum class EMediaTextureSinkMode;
 class FMediaTextureResource
 	: public FRenderTarget
 	, public FTextureResource
-	, public FTickerObjectBase
 {
-	struct FResource
-	{
-		void* LockedData;
-		TRefCountPtr<FRHITexture2D> RenderTarget;
-		TRefCountPtr<FRHITexture2D> ShaderResource;
-
-		FResource() : LockedData(nullptr) { }
-	};
-
 public:
 
 	/** 
 	 * Creates and initializes a new instance.
 	 *
 	 * @param InOwner The Movie texture object to create a resource for (must not be nullptr).
-	 * @param InClearColor The clear color to use.
-	 * @param InOutputDimensions
-	 * @param InSinkFormat
-	 * @param InSinkMode
+	 * @param InSink The sink that receives texture samples from the media player.
 	 */
-	FMediaTextureResource(UMediaTexture& InOwner, const FLinearColor& InClearColor, const FIntPoint& InOutputDimensions, EMediaTextureSinkFormat InSinkFormat, EMediaTextureSinkMode InSinkMode);
+	FMediaTextureResource(UMediaTexture& InOwner, FIntPoint& InOwnerDim, SIZE_T& InOwnerSize);
 
 	/** Virtual destructor. */
-	virtual ~FMediaTextureResource();
+	virtual ~FMediaTextureResource() { }
 
 public:
 
-	/**
-	 * Acquire a pointer to locked texture memory buffer to render to.
-	 *
-	 * @return Pointer to buffer on success, nullptr otherwise.
-	 */
-	void* AcquireBuffer();
-
-	/** Display the latest texture buffer. */
-	void DisplayBuffer();
-
-	/**
-	 * Get the memory size of this resource.
-	 *
-	 * @return Resource size (in bytes).
-	 */
-	DEPRECATED(4.14, "GetResourceSize is deprecated. Please use GetResourceSizeEx or GetResourceSizeBytes instead.")
-	SIZE_T GetResourceSize() const
+	/** Parameters for the Render method. */
+	struct FRenderParams
 	{
-		return GetResourceSizeBytes();
-	}
+		/** The clear color to use when clearing the texture. */
+		FLinearColor ClearColor;
 
-	void GetResourceSizeEx(FResourceSizeEx& CumulativeResourceSize) const
-	{
-		CumulativeResourceSize.AddUnknownMemoryBytes(CachedResourceSizeBytes);
-	}
-
-	SIZE_T GetResourceSizeBytes() const
-	{
-		return CachedResourceSizeBytes;
-	}
+		/** Guid associated with media player. */
+		FGuid PlayerGuid;
 
-	/**
-	 * Get the render target texture to render to.
-	 *
-	 * @return The texture, or nullptr if called in Buffered mode or not on rendering thread.
-	 */
-	FRHITexture* GetTexture();
+		/** The player's play rate. */
+		float Rate;
 
-	/**
-	 * Initialize the render target buffer(s).
-	 *
-	 * @param OutputDim Width and height of the video output (in pixels).
-	 * @param BufferDim Width and height of the sink buffer(s) (in pixels).
-	 * @param Format The pixel format of the sink's render target texture.
-	 * @param Mode The mode to operate the sink in (buffered vs. unbuffered).
-	 */
-	void InitializeBuffer(FIntPoint OutputDim, FIntPoint BufferDim, EMediaTextureSinkFormat Format, EMediaTextureSinkMode Mode);
+		/** The player facade that provides the video samples to render. */
+		TWeakPtr<FMediaTextureSampleSource, ESPMode::ThreadSafe> SampleSource;
 
-	/** Release a previously acquired texture buffer. */
-	void ReleaseBuffer();
+		/** Whether output should be in sRGB color space. */
+		bool SrgbOutput;
 
-	/** Shut down the render target buffers. */
-	void ShutdownBuffer();
+		/** The time of the video frame to render (in player's clock). */
+		FTimespan Time;
+	};
 
 	/**
-	 * Update the render target texture with new frame data.
+	 * Render the texture resource.
 	 *
-	 * @param Data The frame data to copy into the render target.
-	 * @param Pitch Number of bytes per pixel row (0 = default).
-	 */
-	void UpdateBuffer(const uint8* Data, uint32 Pitch);
-
-	/**
-	 * Update the underlying resource textures.
+	 * This method is called on the render thread by the MediaTexture that owns this
+	 * texture resource to clear or redraw the resource using the given parameters.
 	 *
-	 * @param RenderTarget The new render target texture.
-	 * @param ShaderResource The new shader resource texture.
+	 * @param Params Render parameters.
 	 */
-	void UpdateTextures(FRHITexture* RenderTarget, FRHITexture* ShaderResource);
+	void Render(const FRenderParams& Params);
 
 public:
 
@@ -137,145 +92,69 @@ public:
 	virtual void InitDynamicRHI() override;
 	virtual void ReleaseDynamicRHI() override;
 
-public:
-
-	//~ FTickerObjectBase interface
-
-	virtual bool Tick(float DeltaTime) override;
-
 protected:
 
-	/** Cache the size of this resource, so it doesn't need to be recalculated. */
-	void CacheResourceSize();
-
 	/**
-	 * Helper function to convert the given resource to the target resource's pixel format.
+	 * Clear the texture using the given clear color.
 	 *
-	 * @param Resource The texture resource to convert.
-	 * @see DisplayResource
+	 * @param ClearColor The clear color to use.
+	 * @param SrgbOutput Whether the output texture is in sRGB color space.
 	 */
-	void ConvertResource(const FResource& Resource);
+	void ClearTexture(const FLinearColor& ClearColor, bool SrgbOutput);
 
 	/**
-	 * Helper function to display the given resource in the media texture.
+	 * Render the given texture sample by converting it on the GPU.
 	 *
-	 * @param Resource The resource to display.
-	 * @see ConvertResource
+	 * @param Sample The texture sample to convert.
+	 * @param ClearColor The clear color to use for the output texture.
+	 * @param SrgbOutput Whether the output texture is in sRGB color space.
+	 * @see CopySample
 	 */
-	void DisplayResource(const FResource& Resource);
+	void ConvertSample(const TSharedPtr<IMediaTextureSample, ESPMode::ThreadSafe>& Sample, const FLinearColor& ClearColor, bool SrgbOutput);
 
 	/**
-	 * Initialize this resource.
+	 * Render the given texture sample by using it as or copying it to the render target.
 	 *
-	 * @param OutputDim Width and height of the output texture (in pixels).
-	 * @param BufferDim Width and height of the buffer texture(s) (in pixels).
-	 * @param Format The new texture format.
-	 * @param Mode The new sink mode.
+	 * @param Sample The texture sample to copy.
+	 * @param ClearColor The clear color to use for the output texture.
+	 * @param SrgbOutput Whether the output texture is in sRGB color space.
+	 * @see ConvertSample
 	 */
-	void InitializeResource(FIntPoint OutputDim, FIntPoint BufferDim, EMediaTextureSinkFormat Format, EMediaTextureSinkMode Mode);
+	void CopySample(const TSharedPtr<IMediaTextureSample, ESPMode::ThreadSafe>& Sample, const FLinearColor& ClearColor, bool SrgbOutput);
 
-	/** Process any queued up tasks on the render thread. */
-	void ProcessRenderThreadTasks();
+	/** Calculates the current resource size and notifies the owner texture. */
+	void UpdateResourceSize();
 
 	/**
-	 * Helper function to make the given resource the output.
+	 * Set the owner's texture reference to the given texture.
 	 *
-	 * @param Resource The resource to set.
+	 * @param NewTexture The texture to set.
 	 */
-	void SetResource(const FResource& Resource);
-
-	/** Swap the render target resource for display. */
-	void SwapResource();
+	void UpdateTextureReference(FRHITexture2D* NewTexture);
 
 private:
 
-	/** The media texture that owns this resource. */
-	UMediaTexture& Owner;
+	/** Whether the texture has been cleared. */
+	bool Cleared;
 
-	/**
-	 * Texture resources for buffered mode or pixel conversions.
-	 *
-	 * In Buffered mode, all three resources are used for triple buffering.
-	 * If no pixel format conversion is required, the triple buffer's Read
-	 * buffer is used as the output resource. Otherwise the Read buffer is
-	 * converted into a separate output resource.
-	 *
-	 * In Unbuffered mode when pixel format conversion is required, only the
-	 * third buffer (index = 2) is used. Otherwise data is written directly
-	 * to a separate output resource.
-	 */
-	FResource BufferResources[3];
+	/** Tracks the current clear color. */
+	FLinearColor CurrentClearColor;
 
-	/** Triple-buffer for texture resources. */
-	TTripleBuffer<FResource> TripleBuffer;
-
-private:
-
-	//~ The following fields are owned by the render thread
-
-	/** The clear color to use. */
-	FLinearColor BufferClearColor;
-
-	/**
-	 * Width and height of the buffer resources (in pixels).
-	 *
-	 * This may be less than OutputDimensions, because the buffers can
-	 * have formats with multiple pixels packed into a single RGBA tuple.
-	 */
-	FIntPoint BufferDimensions;
+	/** Input render target if the texture samples don't provide one (for conversions). */
+	TRefCountPtr<FRHITexture2D> InputTarget;
 
-	/** Number of bytes per row in buffer resources. */
-	SIZE_T BufferPitch;
+	/** Output render target if the texture samples don't provide one. */
+	TRefCountPtr<FRHITexture2D> OutputTarget;
 
-	/** Total size of this resource.*/
-	SIZE_T CachedResourceSizeBytes;
-
-	/** Width and height of the output resource (in pixels). */
-	FIntPoint OutputDimensions;
-
-	/**
-	 * Asynchronous tasks for the render thread.
-	 *
-	 * We can't reliably enqueue render commands using the existing macros
-	 * from the decoder thread, which runs neither on the game nor on the
-	 * render thread, so we maintain our own queue of render commands here.
-	 */
-	TQueue<TFunction<void()>> RenderThreadTasks;
-
-	/** Whether the current media sink format requires pixel format conversion. */
-	bool RequiresConversion;
-
-	/** The render target's pixel format. */
-	EMediaTextureSinkFormat SinkFormat;
-
-	/** The mode that this sink is currently operating in. */
-	EMediaTextureSinkMode SinkMode;
+	/** The media texture that owns this resource. */
+	UMediaTexture& Owner;
 
-private:
+	/** Reference to the owner's texture dimensions field. */
+	FIntPoint& OwnerDim;
 
-	//~ The following fields are owned by the decoder thread
+	/** Reference to the owner's texture size field. */
+	SIZE_T& OwnerSize;
 
-	/**
-	 * The sink's current state.
-	 *
-	 * The purpose of this field is to prevent the decoder thread from doing
-	 * stupid things, such as calling InitializeBuffer or ShutdownBuffer in
-	 * the wrong order. Strictly speaking, it is not thread-safe, because
-	 * state changing operations are not atomic in the current implementation.
-	 * However, it shouldn't be an issue, because relevant accesses are gated.
-	 *
-	 * State changes on the decoder thread will be visible there immediately.
-	 * There are a couple places where the state is set on the render thread,
-	 * but here we are ok with eventual consistency on the decoder thread.
-	 * The latest state will eventually trickle down, and in the meantime the
-	 * decoder is in a pending state, such as Initializing or ShuttingDown.
-	 */
-	enum class EState
-	{
-		Initializing,
-		Initialized,
-		ShuttingDown,
-		ShutDown
-	}
-	State;
+	/** The current media player facade to get video samples from. */
+	TWeakPtr<FMediaPlayerFacade, ESPMode::ThreadSafe> PlayerFacadePtr;
 };
diff --git a/Engine/Source/Runtime/MediaAssets/Private/Tests/MediaPlaylistTest.cpp b/Engine/Source/Runtime/MediaAssets/Private/Tests/MediaPlaylistTest.cpp
new file mode 100644
--- /dev/null
+++ b/Engine/Source/Runtime/MediaAssets/Private/Tests/MediaPlaylistTest.cpp
@@ -0,0 +1,177 @@
+// Copyright 1998-2017 Epic Games, Inc. All Rights Reserved.
+
+#include "CoreTypes.h"
+#include "MediaPlaylist.h"
+#include "Misc/AutomationTest.h"
+#include "UObject/UObjectGlobals.h"
+
+#if WITH_DEV_AUTOMATION_TESTS
+
+
+IMPLEMENT_SIMPLE_AUTOMATION_TEST(FMediaPlaylistTest, "System.Media.Assets.Playlist", EAutomationTestFlags::EditorContext | EAutomationTestFlags::ClientContext | EAutomationTestFlags::EngineFilter)
+
+
+bool FMediaPlaylistTest::RunTest(const FString& Parameters)
+{
+	auto Playlist = NewObject<UMediaPlaylist>();
+	int32 Index;
+
+	// empty play list
+	{
+		TestEqual(TEXT("A new play list must be empty"), Playlist->Num(), 0);
+
+		Index = INDEX_NONE;
+		TestNull(TEXT("GetNext(INDEX_NONE) on a new play list must return nullptr"), Playlist->GetNext(Index));
+		TestEqual(TEXT("GetNext(INDEX_NONE) on a new play list must yield INDEX_NONE"), Index, INDEX_NONE);
+		TestNull(TEXT("GetPrevious(INDEX_NONE) on a new play list must return nullptr"), Playlist->GetPrevious(Index));
+		TestEqual(TEXT("GetPrevious(INDEX_NONE) on a new play list must yield INDEX_NONE"), Index, INDEX_NONE);
+
+		Index = 0;
+		TestNull(TEXT("GetNext(0) on a new play list must return nullptr"), Playlist->GetNext(Index));
+		TestEqual(TEXT("GetNext(0) on a new play list must yield INDEX_NONE"), Index, INDEX_NONE);
+		TestNull(TEXT("GetPrevious(0) on a new play list must return nullptr"), Playlist->GetPrevious(Index));
+		TestEqual(TEXT("GetPrevious(0) on a new play list must yield INDEX_NONE"), Index, INDEX_NONE);
+
+		Index = 1;
+		TestNull(TEXT("GetNext(1) on a new play list must return nullptr"), Playlist->GetNext(Index));
+		TestEqual(TEXT("GetNext(1) on a new play list must yield INDEX_NONE"), Index, INDEX_NONE);
+		TestNull(TEXT("GetPrevious(1) on a new play list must return nullptr"), Playlist->GetPrevious(Index));
+		TestEqual(TEXT("GetPrevious(1) on a new play list must yield INDEX_NONE"), Index, INDEX_NONE);
+
+		TestNull(TEXT("GetRandom() on a new play list must return nullptr"), Playlist->GetRandom(Index));
+		TestEqual(TEXT("GetRandom() on a new play list must yield INDEX_NONE"), Index, INDEX_NONE);
+	}
+
+	// one entry (no looping)
+	{
+		Playlist->Loop = false;
+
+		Playlist->Add((UMediaSource*)1);
+		TestEqual(TEXT("A play list with one entry and no looping must have length 1"), Playlist->Num(), 1);
+
+		Index = INDEX_NONE;
+		TestNotNull(TEXT("GetNext(INDEX_NONE) on a play list with one entry and no looping must return the first item"), Playlist->GetNext(Index));
+		TestEqual(TEXT("GetNext(INDEX_NONE) on a play list with one entry and no looping must yield 0"), Index, 0);
+
+		Index = INDEX_NONE;
+		TestNull(TEXT("GetPrevious(INDEX_NONE) on a play list with one entry and no looping must return nullptr"), Playlist->GetPrevious(Index));
+		TestEqual(TEXT("GetPrevious(INDEX_NONE) on a play list with one entry and no looping must yield INDEX_NONE"), Index, INDEX_NONE);
+
+		Index = 0;
+		TestNull(TEXT("GetNext(0) on a play list with one entry and no looping must return nullptr"), Playlist->GetNext(Index));
+		TestEqual(TEXT("GetNext(0) on a play list with one entry and no looping must yield INDEX_NONE"), Index, INDEX_NONE);
+
+		Index = 0;
+		TestNull(TEXT("GetPrevious(0) on a play list with one entry and no looping must return nullptr"), Playlist->GetPrevious(Index));
+		TestEqual(TEXT("GetPrevious(0) on a play list with one entry and no looping must yield INDEX_NONE"), Index, INDEX_NONE);
+
+		Index = 1;
+		TestNull(TEXT("GetNext(1) on a play list with one entry and no looping must return nullptr"), Playlist->GetNext(Index));
+		TestEqual(TEXT("GetNext(1) on a play list with one entry and no looping must yield INDEX_NONE"), Index, INDEX_NONE);
+
+		Index = 1;
+		TestNotNull(TEXT("GetPrevious(1) on a play list with one entry and no looping must return the first item"), Playlist->GetPrevious(Index));
+		TestEqual(TEXT("GetPrevious(1) on a play list with one entry and no looping must yield 0"), Index, 0);
+
+		TestNotNull(TEXT("GetRandom() on a play list with one entry and no looping must return the first item"), Playlist->GetRandom(Index));
+		TestEqual(TEXT("GetRandom() on a play list with one entry and no looping must yield 0"), Index, 0);
+	}
+
+	// one entry (looping)
+	{
+		Playlist->Loop = true;
+
+		Index = INDEX_NONE;
+		TestNotNull(TEXT("GetNext(INDEX_NONE) on a play list with one entry and looping must return the first item"), Playlist->GetNext(Index));
+		TestEqual(TEXT("GetNext(INDEX_NONE) on a play list with one entry and looping must yield 0"), Index, 0);
+
+		Index = INDEX_NONE;
+		TestNotNull(TEXT("GetPrevious(INDEX_NONE) on a play list with one entry and looping must return the first item"), Playlist->GetPrevious(Index));
+		TestEqual(TEXT("GetPrevious(INDEX_NONE) on a play list with one entry and looping must yield 0"), Index, 0);
+
+		Index = 0;
+		TestNotNull(TEXT("GetNext(0) on a play list with one entry and looping must return the first item"), Playlist->GetNext(Index));
+		TestEqual(TEXT("GetNext(0) on a play list with one entry and looping must yield 0"), Index, 0);
+
+		Index = 0;
+		TestNotNull(TEXT("GetPrevious(0) on a play list with one entry and looping must return the first item"), Playlist->GetPrevious(Index));
+		TestEqual(TEXT("GetPrevious(0) on a play list with one entry and looping must yield 0"), Index, 0);
+
+		Index = 1;
+		TestNotNull(TEXT("GetNext(1) on a play list with one entry and looping must return the first item"), Playlist->GetNext(Index));
+		TestEqual(TEXT("GetNext(1) on a play list with one entry and looping must yield 0"), Index, 0);
+
+		Index = 1;
+		TestNotNull(TEXT("GetPrevious(1) on a play list with one entry and looping must return the first item"), Playlist->GetPrevious(Index));
+		TestEqual(TEXT("GetPrevious(1) on a play list with one entry and looping must yield 0"), Index, 0);
+	}
+
+	// two entries (no looping)
+	{
+		Playlist->Loop = false;
+
+		Playlist->Add((UMediaSource*)2);
+		TestEqual(TEXT("A play list with two entries and no looping must have length 2"), Playlist->Num(), 2);
+
+		Index = INDEX_NONE;
+		TestEqual(TEXT("GetNext(INDEX_NONE) on a play list with two entries and no looping must return the first item"), Playlist->GetNext(Index), (UMediaSource*)1);
+		TestEqual(TEXT("GetNext(INDEX_NONE) on a play list with two entries and no looping must yield 0"), Index, 0);
+
+		Index = INDEX_NONE;
+		TestNull(TEXT("GetPrevious(INDEX_NONE) on a play list with two entries and no looping must return the first item"), Playlist->GetPrevious(Index));
+		TestEqual(TEXT("GetPrevious(INDEX_NONE) on a play list with two entries and no looping must yield INDEX_NONE"), Index, INDEX_NONE);
+
+		Index = 0;
+		TestEqual(TEXT("GetNext(0) on a play list with two entries and no looping must return the second item"), Playlist->GetNext(Index), (UMediaSource*)2);
+		TestEqual(TEXT("GetNext(0) on a play list with two entries and no looping must yield 1"), Index, 1);
+
+		Index = 0;
+		TestNull(TEXT("GetPrevious(0) on a play list with two entries and no looping must return nullptr"), Playlist->GetPrevious(Index));
+		TestEqual(TEXT("GetPrevious(0) on a play list with two entries and no looping must yield INDEX_NONE"), Index, INDEX_NONE);
+
+		Index = 1;
+		TestNull(TEXT("GetNext(1) on a play list with two entries and no looping must return nullptr"), Playlist->GetNext(Index));
+		TestEqual(TEXT("GetNext(1) on a play list with two entries and no looping must yield INDEX_NONE"), Index, INDEX_NONE);
+
+		Index = 1;
+		TestEqual(TEXT("GetPrevious(1) on a play list with two entries and no looping must return nullptr"), Playlist->GetPrevious(Index), (UMediaSource*)1);
+		TestEqual(TEXT("GetPrevious(1) on a play list with two entries and no looping must yield 0"), Index, 0);
+	}
+
+	// two entries (looping)
+	{
+		Playlist->Loop = true;
+
+		Index = INDEX_NONE;
+		TestEqual(TEXT("GetNext(INDEX_NONE) on a play list with two entries and looping must return the first item"), Playlist->GetNext(Index), (UMediaSource*)1);
+		TestEqual(TEXT("GetNext(INDEX_NONE) on a play list with two entries and looping must yield 0"), Index, 0);
+
+		Index = INDEX_NONE;
+		TestEqual(TEXT("GetPrevious(INDEX_NONE) on a play list with two entries and looping must return the second item"), Playlist->GetPrevious(Index), (UMediaSource*)2);
+		TestEqual(TEXT("GetPrevious(INDEX_NONE) on a play list with two entries and looping must yield 1"), Index, 1);
+
+		Index = 0;
+		TestEqual(TEXT("GetNext(0) on a play list with two entries and looping must return the second item"), Playlist->GetNext(Index), (UMediaSource*)2);
+		TestEqual(TEXT("GetNext(0) on a play list with two entries and looping must yield 1"), Index, 1);
+
+		Index = 0;
+		TestEqual(TEXT("GetPrevious(0) on a play list with two entries and looping must return the first item"), Playlist->GetPrevious(Index), (UMediaSource*)2);
+		TestEqual(TEXT("GetPrevious(0) on a play list with two entries and looping must yield 1"), Index, 1);
+
+		Index = 1;
+		TestEqual(TEXT("GetNext(1) on a play list with two entries and looping must return the first item"), Playlist->GetNext(Index), (UMediaSource*)1);
+		TestEqual(TEXT("GetNext(1) on a play list with two entries and looping must yield 0"), Index, 0);
+
+		Index = 1;
+		TestEqual(TEXT("GetPrevious(1) on a play list with two entries and looping must return the first item"), Playlist->GetPrevious(Index), (UMediaSource*)1);
+		TestEqual(TEXT("GetPrevious(1) on a play list with two entries and looping must yield 0"), Index, 0);
+
+		TestNotNull(TEXT("GetRandom() on a play list with two entries and no looping must not return nullptr"), Playlist->GetRandom(Index));
+		TestNotEqual(TEXT("GetRandom() on a play list with two entries and no looping must not yield INDEX_NONE"), Index, (int32)INDEX_NONE);
+	}
+
+	return true;
+}
+
+
+#endif //WITH_DEV_AUTOMATION_TESTS
diff --git a/Engine/Source/Runtime/MediaAssets/Public/BaseMediaSource.h b/Engine/Source/Runtime/MediaAssets/Public/BaseMediaSource.h
--- a/Engine/Source/Runtime/MediaAssets/Public/BaseMediaSource.h
+++ b/Engine/Source/Runtime/MediaAssets/Public/BaseMediaSource.h
@@ -2,7 +2,9 @@
 
 #pragma once
 
-#include "CoreMinimal.h"
+#include "CoreTypes.h"
+#include "Containers/Map.h"
+#include "Containers/UnrealString.h"
 #include "UObject/ObjectMacros.h"
 #include "UObject/Object.h"
 #include "MediaSource.h"
diff --git a/Engine/Source/Runtime/MediaAssets/Public/FileMediaSource.h b/Engine/Source/Runtime/MediaAssets/Public/FileMediaSource.h
--- a/Engine/Source/Runtime/MediaAssets/Public/FileMediaSource.h
+++ b/Engine/Source/Runtime/MediaAssets/Public/FileMediaSource.h
@@ -2,7 +2,8 @@
 
 #pragma once
 
-#include "CoreMinimal.h"
+#include "CoreTypes.h"
+#include "Containers/UnrealString.h"
 #include "UObject/ObjectMacros.h"
 #include "UObject/ScriptMacros.h"
 #include "BaseMediaSource.h"
@@ -36,13 +37,21 @@ public:
 	 * Get the path to the media file to be played.
 	 *
 	 * @return The file path.
-	 * @see SetFilePath
+	 * @see GetFullPath, SetFilePath
 	 */
 	const FString& GetFilePath() const
 	{
 		return FilePath;
 	}
 
+	/**
+	 * Get the full path to the file.
+	 *
+	 * @return The full file path.
+	 * @return GetFilePath
+	 */
+	FString GetFullPath() const;
+
 	/**
 	 * Set the path to the media file that this source represents.
 	 *
@@ -73,9 +82,4 @@ public:
 
 	/** Name of the PrecacheFile media option. */
 	static FName PrecacheFileOption;
-
-protected:
-
-	/** Get the full path to the file. */
-	FString GetFullPath() const;
 };
diff --git a/Engine/Source/Runtime/MediaAssets/Public/MediaPlayer.h b/Engine/Source/Runtime/MediaAssets/Public/MediaPlayer.h
--- a/Engine/Source/Runtime/MediaAssets/Public/MediaPlayer.h
+++ b/Engine/Source/Runtime/MediaAssets/Public/MediaPlayer.h
@@ -2,25 +2,27 @@
 
 #pragma once
 
-#include "CoreMinimal.h"
-#include "Containers/Ticker.h"
+#include "CoreTypes.h"
+#include "Containers/UnrealString.h"
+#include "Delegates/Delegate.h"
+#include "Math/Quat.h"
+#include "Math/Rotator.h"
+#include "Templates/SharedPointer.h"
 #include "UObject/ObjectMacros.h"
 #include "UObject/Object.h"
 #include "UObject/ScriptMacros.h"
-#include "IMediaOverlaySink.h"
-#include "MediaPlayerBase.h"
-#include "MediaPlayer.generated.h"
+#include "Misc/Guid.h"
 
+#include "MediaPlayer.generated.h"
 
-class FMediaPlayerBase;
+class FMediaPlayerFacade;
 class IMediaPlayer;
-class UMediaOverlays;
 class UMediaPlaylist;
-class UMediaSoundWave;
 class UMediaSource;
-class UMediaTexture;
+
 enum class EMediaEvent;
 
+
 /** Multicast delegate that is invoked when a media event occurred in the player. */
 DECLARE_DYNAMIC_MULTICAST_DELEGATE(FOnMediaPlayerMediaEvent);
 
@@ -42,12 +44,12 @@ enum class EMediaPlayerTrack : uint8
 	/** Audio track. */
 	Audio,
 
-	/** Binary data track. */
-	Binary,
-
 	/** Caption track. */
 	Caption,
 
+	/** Metadata track. */
+	Metadata,
+
 	/** Script track. */
 	Script,
 
@@ -68,7 +70,6 @@ enum class EMediaPlayerTrack : uint8
 UCLASS(BlueprintType, hidecategories=(Object))
 class MEDIAASSETS_API UMediaPlayer
 	: public UObject
-	, public FTickerObjectBase
 {
 	GENERATED_UCLASS_BODY()
 
@@ -118,6 +119,39 @@ public:
 	UFUNCTION(BlueprintCallable, Category="Media|MediaPlayer")
 	void Close();
 
+	/**
+	 * Get the number of channels in the specified audio track.
+	 *
+	 * @param TrackIndex Index of the audio track, or INDEX_NONE for the selected one.
+	 * @param FormatIndex Index of the track format, or INDEX_NONE for the selected one.
+	 * @return Number of channels.
+	 * @see GetAudioTrackSampleRate, GetAudioTrackType
+	 */
+	UFUNCTION(BlueprintCallable, Category = "Media|MediaPlayer")
+	int32 GetAudioTrackChannels(int32 TrackIndex, int32 FormatIndex) const;
+
+	/**
+	 * Get the sample rate of the specified audio track.
+	 *
+	 * @param TrackIndex Index of the audio track, or INDEX_NONE for the selected one.
+	 * @param FormatIndex Index of the track format, or INDEX_NONE for the selected one.
+	 * @return Samples per second.
+	 * @see GetAudioTrackChannels, GetAudioTrackType
+	 */
+	UFUNCTION(BlueprintCallable, Category = "Media|MediaPlayer")
+	int32 GetAudioTrackSampleRate(int32 TrackIndex, int32 FormatIndex) const;
+
+	/**
+	 * Get the type of the specified audio track format.
+	 *
+	 * @param TrackIndex The index of the track, or INDEX_NONE for the selected one.
+	 * @param FormatIndex Index of the track format, or INDEX_NONE for the selected one.
+	 * @return Audio format type string.
+	 * @see GetAudioTrackSampleRate, GetAudioTrackSampleRate
+	 */
+	UFUNCTION(BlueprintCallable, Category="Media|MediaPlayer")
+	FString GetAudioTrackType(int32 TrackIndex, int32 FormatIndex) const;
+
 	/**
 	 * Get the name of the current desired native player.
 	 *
@@ -137,25 +171,35 @@ public:
 	FTimespan GetDuration() const;
 
 	/**
-	 * Get the supported forward playback rates.
+	 * Get the current horizontal field of view (only for 360 videos).
 	 *
-	 * @param Unthinned Whether the rates are for unthinned playback (default = true).
-	 * @return The range of supported rates.
-	 * @see GetReverseRates
+	 * @return Horizontal field of view (in Euler degrees).
+	 * @see GetVerticalFieldOfView, GetViewRotation, SetHorizontalFieldOfView
 	 */
 	UFUNCTION(BlueprintCallable, Category="Media|MediaPlayer")
-	FFloatRange GetForwardRates(bool Unthinned = true);
+	float GetHorizontalFieldOfView() const;
 
 	/**
 	 * Get the number of tracks of the given type.
 	 *
 	 * @param TrackType The type of media tracks.
 	 * @return Number of tracks.
-	 * @see GetSelectedTrack, SelectTrack
+	 * @see GetNumTrackFormats, GetSelectedTrack, SelectTrack
 	 */
 	UFUNCTION(BlueprintCallable, Category="Media|MediaPlayer")
 	int32 GetNumTracks(EMediaPlayerTrack TrackType) const;
 
+	/**
+	 * Get the number of formats of the specified track.
+	 *
+	 * @param TrackType The type of media tracks.
+	 * @param TrackIndex The index of the track.
+	 * @return Number of formats.
+	 * @see GetNumTracks, GetSelectedTrack, SelectTrack
+	 */
+	UFUNCTION(BlueprintCallable, Category="Media|MediaPlayer")
+	int32 GetNumTrackFormats(EMediaPlayerTrack TrackType, int32 TrackIndex) const;
+
 	/**
 	 * Get the name of the current native media player.
 	 *
@@ -165,34 +209,61 @@ public:
 	FName GetPlayerName() const;
 
 	/**
-	 * Get the media's current playback rate.
+	 * Get the current play list.
 	 *
-	 * @return The playback rate.
-	 * @see SetRate, SupportsRate
+	 * Media players always have a valid play list. In C++ code you can use
+	 * the GetPlaylistRef to get a reference instead of a pointer to it.
+	 *
+	 * @return The play list.
+	 * @see GetPlaylistIndex, GetPlaylistRef
 	 */
 	UFUNCTION(BlueprintCallable, Category="Media|MediaPlayer")
-	float GetRate() const;
+	UMediaPlaylist* GetPlaylist() const
+	{
+		return Playlist;
+	}
+
+	/**
+	 * Get the current play list index.
+	 *
+	 * @return Play list index.
+	 * @see GetPlaylist
+	 */
+	UFUNCTION(BlueprintCallable, Category = "Media|MediaPlayer")
+	int32 GetPlaylistIndex() const
+	{
+		return PlaylistIndex;
+	}
 
 	/**
-	 * Get the supported reverse playback rates.
+	 * Get the media's current playback rate.
 	 *
-	 * @param Unthinned Whether the rates are for unthinned playback (default = true).
-	 * @return The range of supported rates.
-	 * @see GetForwardRates
+	 * @return The playback rate.
+	 * @see SetRate, SupportsRate
 	 */
 	UFUNCTION(BlueprintCallable, Category="Media|MediaPlayer")
-	FFloatRange GetReverseRates(bool Unthinned = true);
+	float GetRate() const;
 
 	/**
 	 * Get the index of the currently selected track of the given type.
 	 *
 	 * @param TrackType The type of track to get.
 	 * @return The index of the selected track, or INDEX_NONE if no track is active.
-	 * @see GetNumTracks, SelectTrack
+	 * @see GetNumTracks, GetTrackFormat, SelectTrack
 	 */
 	UFUNCTION(BlueprintCallable, Category="Media|MediaPlayer")
 	int32 GetSelectedTrack(EMediaPlayerTrack TrackType) const;
 
+	/**
+	 * Get the supported playback rates.
+	 *
+	 * @param Unthinned Whether the rates are for unthinned playback.
+	 * @param Will contain the the ranges of supported rates.
+	 * @see SetRate, SupportsRate
+	 */
+	UFUNCTION(BlueprintCallable, Category="Media|MediaPlayer")
+	void GetSupportedRates(TArray<FFloatRange>& OutRates, bool Unthinned) const;
+
 	/**
 	 * Get the media's current playback time.
 	 *
@@ -206,18 +277,29 @@ public:
 	 * Get the human readable name of the specified track.
 	 *
 	 * @param TrackType The type of track.
-	 * @param TrackIndex The index of the track.
+	 * @param TrackIndex The index of the track, or INDEX_NONE for the selected one.
 	 * @return Display name.
 	 * @see GetNumTracks, GetTrackLanguage
 	 */
 	UFUNCTION(BlueprintCallable, Category = "Media|MediaPlayer")
 	FText GetTrackDisplayName(EMediaPlayerTrack TrackType, int32 TrackIndex) const;
 
+	/**
+	 * Get the index of the active format of the specified track type.
+	 *
+	 * @param TrackType The type of track.
+	 * @param TrackIndex The index of the track, or INDEX_NONE for the selected one.
+	 * @return The index of the selected format.
+	 * @see GetNumTrackFormats, GetSelectedTrack, SetTrackFormat
+	 */
+	UFUNCTION(BlueprintCallable, Category="Media|MediaPlayer")
+	int32 GetTrackFormat(EMediaPlayerTrack TrackType, int32 TrackIndex) const;
+
 	/**
 	 * Get the language tag of the specified track.
 	 *
 	 * @param TrackType The type of track.
-	 * @param TrackIndex The index of the track.
+	 * @param TrackIndex The index of the track, or INDEX_NONE for the selected one.
 	 * @return Language tag, i.e. "en-US" for English, or "und" for undefined.
 	 * @see GetNumTracks, GetTrackDisplayName
 	 */
@@ -234,49 +316,153 @@ public:
 	const FString& GetUrl() const;
 
 	/**
-	 * Checks whether playback is looping.
+	 * Get the current vertical field of view (only for 360 videos).
+	 *
+	 * @return Vertical field of view (in Euler degrees), or 0.0 if not available.
+	 * @see GetHorizontalFieldOfView, GetViewRotation, SetVerticalFieldOfView
+	 */
+	UFUNCTION(BlueprintCallable, Category="Media|MediaPlayer")
+	float GetVerticalFieldOfView() const;
+
+	/**
+	 * Get the aspect ratio of the specified video track.
+	 *
+	 * @param TrackIndex Index of the video track, or INDEX_NONE for the selected one.
+	 * @param FormatIndex Index of the track format, or INDEX_NONE for the selected one.
+	 * @return Aspect ratio.
+	 * @see GetVideoTrackDimensions, GetVideoTrackFrameRate, GetVideoTrackFrameRates, GetVideoTrackType
+	 */
+	UFUNCTION(BlueprintCallable, Category = "Media|MediaPlayer")
+	float GetVideoTrackAspectRatio(int32 TrackIndex, int32 FormatIndex) const;
+
+	/**
+	 * Get the current dimensions of the specified video track.
+	 *
+	 * @param TrackIndex The index of the track, or INDEX_NONE for the selected one.
+	 * @param FormatIndex Index of the track format, or INDEX_NONE for the selected one.
+	 * @return Video dimensions (in pixels).
+	 * @see GetVideoTrackAspectRatio, GetVideoTrackFrameRate, GetVideoTrackFrameRates, GetVideoTrackType
+	 */
+	UFUNCTION(BlueprintCallable, Category="Media|MediaPlayer")
+	FIntPoint GetVideoTrackDimensions(int32 TrackIndex, int32 FormatIndex) const;
+
+	/**
+	 * Get the frame rate of the specified video track.
+	 *
+	 * @param TrackIndex The index of the track, or INDEX_NONE for the selected one.
+	 * @param FormatIndex Index of the track format, or INDEX_NONE for the selected one.
+	 * @return Frame rate (in frames per second).
+	 * @see GetVideoTrackAspectRatio, GetVideoTrackDimensions, GetVideoTrackFrameRates, GetVideoTrackType, SetVideoTrackFrameRate
+	 */
+	UFUNCTION(BlueprintCallable, Category="Media|MediaPlayer")
+	float GetVideoTrackFrameRate(int32 TrackIndex, int32 FormatIndex) const;
+
+	/**
+	 * Get the supported range of frame rates of the specified video track.
+	 *
+	 * @param TrackIndex The index of the track, or INDEX_NONE for the selected one.
+	 * @param FormatIndex Index of the track format, or INDEX_NONE for the selected one.
+	 * @return Frame rate range (in frames per second).
+	 * @see GetVideoTrackAspectRatio, GetVideoTrackDimensions, GetVideoTrackFrameRate, GetVideoTrackType
+	 */
+	UFUNCTION(BlueprintCallable, Category="Media|MediaPlayer")
+	FFloatRange GetVideoTrackFrameRates(int32 TrackIndex, int32 FormatIndex) const;
+
+	/**
+	 * Get the type of the specified video track format.
+	 *
+	 * @param TrackIndex The index of the track, or INDEX_NONE for the selected one.
+	 * @param FormatIndex Index of the track format, or INDEX_NONE for the selected one.
+	 * @return Video format type string.
+	 * @see GetVideoTrackAspectRatio, GetVideoTrackDimensions, GetVideoTrackFrameRate, GetVideoTrackFrameRates
+	 */
+	UFUNCTION(BlueprintCallable, Category="Media|MediaPlayer")
+	FString GetVideoTrackType(int32 TrackIndex, int32 FormatIndex) const;
+
+	/**
+	 * Get the current view rotation (only for 360 videos).
+	 *
+	 * @return View rotation, or zero rotator if not available.
+	 * @see GetHorizontalFieldOfView, GetVerticalFieldOfView, SetViewRotation
+	 */
+	UFUNCTION(BlueprintCallable, Category="Media|MediaPlayer")
+	FRotator GetViewRotation() const;
+
+	/**
+	 * Check whether the player is in an error state.
+	 *
+	 * When the player is in an error state, no further operations are possible.
+	 * The current media must be closed, and a new media source must be opened
+	 * before the player can be used again. Errors are usually caused by faulty
+	 * media files or interrupted network connections.
+	 *
+	 * @see IsReady
+	 */
+	UFUNCTION(BlueprintCallable, Category="Media|MediaPlayer")
+	bool HasError() const;
+
+	/**
+	 * Check whether playback is buffering data.
 	 *
 	 * @return true if looping, false otherwise.
-	 * @see SetLooping
+	 * @see IsConnecting, IsLooping, IsPaused, IsPlaying, IsPreparing, IsReady
+	 */
+	UFUNCTION(BlueprintCallable, Category="Media|MediaPlayer")
+	bool IsBuffering() const;
+
+	/**
+	 * Check whether the player is currently connecting to a media source.
+	 *
+	 * @return true if connecting, false otherwise.
+	 * @see IsBuffering, IsLooping, IsPaused, IsPlaying, IsPreparing, IsReady
+	 */
+	UFUNCTION(BlueprintCallable, Category="Media|MediaPlayer")
+	bool IsConnecting() const;
+
+	/**
+	 * Check whether playback is looping.
+	 *
+	 * @return true if looping, false otherwise.
+	 * @see IsBuffering, IsConnecting, IsPaused, IsPlaying, IsPreparing, IsReady, SetLooping
 	 */
 	UFUNCTION(BlueprintCallable, Category="Media|MediaPlayer")
 	bool IsLooping() const;
 
 	/**
-	 * Checks whether playback is currently paused.
+	 * Check whether playback is currently paused.
 	 *
 	 * @return true if playback is paused, false otherwise.
-	 * @see CanPause, IsPlaying, IsReady, Pause
+	 * @see CanPause, IsBuffering, IsConnecting, IsLooping, IsPaused, IsPlaying, IsPreparing, IsReady, Pause
 	 */
 	UFUNCTION(BlueprintCallable, Category="Media|MediaPlayer")
 	bool IsPaused() const;
 
 	/**
-	 * Checks whether playback has started.
+	 * Check whether playback has started.
 	 *
 	 * @return true if playback has started, false otherwise.
-	 * @see CanPlay, IsPaused, IsReady, Play
+	 * @see CanPlay, IsBuffering, IsConnecting, IsLooping, IsPaused, IsPlaying, IsPreparing, IsReady, Play
 	 */
 	UFUNCTION(BlueprintCallable, Category="Media|MediaPlayer")
 	bool IsPlaying() const;
 
 	/**
-	 * Checks whether the media is currently opening or buffering.
+	 * Check whether the media is currently opening or buffering.
 	 *
 	 * @return true if playback is being prepared, false otherwise.
-	 * @see CanPlay, IsPaused, IsReady, Play
+	 * @see CanPlay, IsBuffering, IsConnecting, IsLooping, IsPaused, IsPlaying, IsReady, Play
 	 */
 	UFUNCTION(BlueprintCallable, Category="Media|MediaPlayer")
 	bool IsPreparing() const;
 
 	/**
-	 * Checks whether media is ready for playback.
+	 * Check whether media is ready for playback.
 	 *
 	 * A player is ready for playback if it has a media source opened that
 	 * finished preparing and is not in an error state.
 	 *
 	 * @return true if media is ready, false otherwise.
-	 * @see IsPaused, IsPlaying, Stop
+	 * @see HasError, IsBuffering, IsConnecting, IsLooping, IsPaused, IsPlaying, IsPreparing
 	 */
 	UFUNCTION(BlueprintCallable, Category="Media|MediaPlayer")
 	bool IsReady() const;
@@ -429,12 +615,14 @@ public:
 	/**
 	 * Select the active track of the given type.
 	 *
-	 * Only one track of a given type can be active at any time.
+	 * The selected track will use its currently active format. Active formats will
+	 * be remembered on a per track basis. The first available format is active by
+	 * default. To switch the track format, use SetTrackFormat instead.
 	 *
 	 * @param TrackType The type of track to select.
 	 * @param TrackIndex The index of the track to select, or INDEX_NONE to deselect.
 	 * @return true if the track was selected, false otherwise.
-	 * @see GetNumTracks, GetSelectedTrack
+	 * @see GetNumTracks, GetSelectedTrack, SetTrackFormat
 	 */
 	UFUNCTION(BlueprintCallable, Category="Media|MediaPlayer")
 	bool SelectTrack(EMediaPlayerTrack TrackType, int32 TrackIndex);
@@ -458,15 +646,6 @@ public:
 	UFUNCTION(BlueprintCallable, Category="Media|MediaPlayer")
 	bool SetLooping(bool Looping);
 
-	/**
-	 * Assign the given overlays asset to the player's overlay sink.
-	 *
-	 * @param NewOverlays The overlays asset to set.
-	 * @see SetVideoTexture
-	 */
-	UFUNCTION(BlueprintCallable, Category="Media|MediaPlayer")
-	void SetOverlays(UMediaOverlays* NewOverlays);
-
 	/**
 	 * Changes the media's playback rate.
 	 *
@@ -478,25 +657,58 @@ public:
 	bool SetRate(float Rate);
 
 	/**
-	 * Assign the given sound wave to the player's audio sink.
+	 * Set the format on the specified track.
+	 *
+	 * Selecting the format will not switch to the specified track. To switch
+	 * tracks, use SelectTrack instead. If the track is already selected, the
+	 * format change will be applied immediately.
 	 *
-	 * @param NewSoundWave The sound wave to set.
-	 * @see SetVideoTexture
+	 * @param TrackType The type of track to update.
+	 * @param TrackIndex The index of the track to update.
+	 * @param FormatIndex The index of the format to select (must be valid).
+	 * @return true if the track was selected, false otherwise.
+	 * @see GetNumTrackFormats, GetNumTracks, GetTrackFormat, SelectTrack
 	 */
 	UFUNCTION(BlueprintCallable, Category="Media|MediaPlayer")
-	void SetSoundWave(UMediaSoundWave* NewSoundWave);
+	bool SetTrackFormat(EMediaPlayerTrack TrackType, int32 TrackIndex, int32 FormatIndex);
 
 	/**
-	 * Assign the given texture to the player's video sink.
+	 * Set the frame rate of the specified video track.
 	 *
-	 * @param NewTexture The texture to set.
-	 * @see SetSoundWave
+	 * @param TrackIndex The index of the track, or INDEX_NONE for the selected one.
+	 * @param FormatIndex Index of the track format, or INDEX_NONE for the selected one.
+	 * @param FrameRate The frame rate to set (must be in range of format's supported frame rates).
+	 * @return true on success, false otherwise.
+	 * @see GetVideoTrackAspectRatio, GetVideoTrackDimensions, GetVideoTrackFrameRate, GetVideoTrackFrameRates, GetVideoTrackType
 	 */
 	UFUNCTION(BlueprintCallable, Category="Media|MediaPlayer")
-	void SetVideoTexture(UMediaTexture* NewTexture);
+	bool SetVideoTrackFrameRate(int32 TrackIndex, int32 FormatIndex, float FrameRate);
+
+	/**
+	 * Set the field of view (only for 360 videos).
+	 *
+	 * @param Horizontal Horizontal field of view (in Euler degrees).
+	 * @param Vertical Vertical field of view (in Euler degrees).
+	 * @param Whether the field of view change should be absolute (true) or relative (false).
+	 * @return true on success, false otherwise.
+	 * @see GetHorizontalFieldOfView, GetVerticalFieldOfView, SetViewRotation
+	 */
+	UFUNCTION(BlueprintCallable, Category = "Media|MediaPlayer")
+	bool SetViewField(float Horizontal, float Vertical, bool Absolute);
+
+	/**
+	 * Set the view's rotation (only for 360 videos).
+	 *
+	 * @param Rotation The desired view rotation.
+	 * @param Whether the rotation change should be absolute (true) or relative (false).
+	 * @return true on success, false otherwise.
+	 * @see GetViewRotation, SetViewField
+	 */
+	UFUNCTION(BlueprintCallable, Category = "Media|MediaPlayer")
+	bool SetViewRotation(const FRotator& Rotation, bool Absolute);
 
 	/**
-	 * Checks whether the specified playback rate is supported.
+	 * Check whether the specified playback rate is supported.
 	 *
 	 * @param Rate The playback rate to check.
 	 * @param Unthinned Whether no frames should be dropped at the given rate.
@@ -506,7 +718,7 @@ public:
 	bool SupportsRate(float Rate, bool Unthinned) const;
 
 	/**
-	 * Checks whether the currently loaded media supports scrubbing.
+	 * Check whether the currently loaded media supports scrubbing.
 	 *
 	 * @return true if scrubbing is supported, false otherwise.
 	 * @see SupportsRate, SupportsSeeking
@@ -515,7 +727,7 @@ public:
 	bool SupportsScrubbing() const;
 
 	/**
-	 * Checks whether the currently loaded media can jump to a certain position.
+	 * Check whether the currently loaded media can jump to a certain position.
 	 *
 	 * @return true if seeking is supported, false otherwise.
 	 * @see SupportsRate, SupportsScrubbing
@@ -538,7 +750,9 @@ public:
 	 *
 	 * Depending on whether the underlying player implementation opens the media
 	 * synchronously or asynchronously, this event may be executed before or
-	 * after the call top OpenSource / OpenUrl returns.
+	 * after the call to OpenSource / OpenUrl returns.
+	 *
+	 * @see OnMediaOpenFailed
 	 */
 	UPROPERTY(BlueprintAssignable, Category="Media|MediaPlayer")
 	FOnMediaPlayerMediaOpened OnMediaOpened;
@@ -549,97 +763,93 @@ public:
 	 * This delegate is only executed if OpenSource / OpenUrl returned true and
 	 * the media failed to open asynchronously later. It is not executed if
 	 * OpenSource / OpenUrl returned false, indicating an immediate failure.
+	 *
+	 * @see OnMediaOpened
 	 */
 	UPROPERTY(BlueprintAssignable, Category="Media|MediaPlayer")
 	FOnMediaPlayerMediaOpenFailed OnMediaOpenFailed;
 
-	/** A delegate that is invoked when media playback has been resumed. */
+	/**
+	 * A delegate that is invoked when media playback has been resumed.
+	 *
+	 * @see OnPlaybackSuspended
+	 */
 	UPROPERTY(BlueprintAssignable, Category="Media|MediaPlayer")
 	FOnMediaPlayerMediaEvent OnPlaybackResumed;
 
-	/** A delegate that is invoked when media playback has been suspended. */
+	/**
+	 * A delegate that is invoked when media playback has been suspended.
+	 *
+	 * @see OnPlaybackResumed
+	 */
 	UPROPERTY(BlueprintAssignable, Category="Media|MediaPlayer")
 	FOnMediaPlayerMediaEvent OnPlaybackSuspended;
 
-public:
-
 	/**
-	 * Get the player implementation of this object.
+	 * A delegate that is invoked when a seek operation completed successfully.
 	 *
-	 * @return The player implementation.
+	 * Depending on whether the underlying player implementation performs seeks
+	 * synchronously or asynchronously, this event may be executed before or
+	 * after the call to Seek.
 	 */
-	FMediaPlayerBase& GetBasePlayer();
+	UPROPERTY(BlueprintAssignable, Category = "Media|MediaPlayer")
+	FOnMediaPlayerMediaEvent OnSeekCompleted;
+
+public:
 
 	/**
-	 * Get the URL of the media that was last attempted to be opened.
+	 * Get the Guid associated with this media player
 	 *
-	 * @return The last attempted URL.
+	 * @return The Guid.
 	 */
-	const FString& GetLastUrl()
+	const FGuid& GetGuid()
 	{
-		return LastUrl;
+		return PlayerGuid;
 	}
 
 	/**
-	 * Get the overlays assets that receives the captions and subtitles.
+	 * Get the media player facade that manages low-level media players
 	 *
-	 * @return Overlays asset.
-	 * @see GetPlayer, GetPlaylist, GetVideoTexture
+	 * @return The media player facade.
 	 */
-	UMediaOverlays* GetOverlays() const
-	{
-		return Overlays;
-	}
+	TSharedRef<FMediaPlayerFacade, ESPMode::ThreadSafe> GetPlayerFacade() const;
 
 	/**
-	 * Get the currently active play list.
+	 * Get the current play list.
 	 *
-	 * @return The play list, if any.
-	 * @see GetOverlays, GetPlayer, GetPlaylistIndex
+	 * @return The play list.
+	 * @see GetPlaylistIndex, GetPlaylist
 	 */
-	UMediaPlaylist* GetPlaylist() const
+	UMediaPlaylist& GetPlaylistRef() const
 	{
-		return Playlist;
+		check(Playlist != nullptr);
+		return *Playlist;
 	}
 
 	/**
-	 * Get the current play list index.
+	 * Get an event delegate that is invoked when a media event occurred.
 	 *
-	 * @return Play list index.
-	 * @see GetOverlays, GetPlayer, GetPlaylist
+	 * @return The delegate.
 	 */
-	int32 GetPlaylistIndex() const
+	DECLARE_EVENT_OneParam(UMediaPlayer, FOnMediaEvent, EMediaEvent /*Event*/)
+	FOnMediaEvent& OnMediaEvent()
 	{
-		return PlaylistIndex;
+		return MediaEvent;
 	}
 
+#if WITH_EDITOR
 	/**
-	 * Get the sound wave that receives the audio output.
+	 * Called when PIE has been paused.
 	 *
-	 * @return Sound wave asset.
-	 * @see GetOverlays, GetPlayer, GetVideoTexture
+	 * @see ResumePIE
 	 */
-	UMediaSoundWave* GetSoundWave() const
-	{
-		return SoundWave;
-	}
+	void PausePIE();
 
 	/**
-	 * Get the media texture that receives the video output.
+	 * Called when PIE has been resumed.
 	 *
-	 * @return Media texture asset.
-	 * @see GetOverlays, GetPlayer, GetSoundWave
+	 * @see PausePIE
 	 */
-	UMediaTexture* GetVideoTexture() const
-	{
-		return VideoTexture;
-	}
-
-#if WITH_EDITOR
-	/** Called when PIE has been paused. */
-	void PausePIE();
-
-	/** Called when PIE has been resumed. */
 	void ResumePIE();
 #endif
 
@@ -648,90 +858,170 @@ public:
 	//~ UObject interface
 
 	virtual void BeginDestroy() override;
+	virtual bool CanBeInCluster() const override;
 	virtual FString GetDesc() override;
-	virtual void PostLoad() override; 
-	virtual bool CanBeInCluster() const override { return false; }
+	virtual void PostDuplicate(bool bDuplicateForPIE) override;
+	virtual void PostLoad() override;
 
 #if WITH_EDITOR
 	virtual void PostEditChangeProperty(FPropertyChangedEvent& PropertyChangedEvent) override;
-	virtual void PreEditChange(UProperty* PropertyAboutToChange) override;
 #endif
 
 public:
 
-	//~ FTickerObjectBase interface
+	/**
+	 * Duration of samples to cache ahead of the play head.
+	 *
+	 * @see CacheBehind, CacheBehindGame
+	 */
+	UPROPERTY(BlueprintReadWrite, Category=Caching)
+	FTimespan CacheAhead;
 
-	virtual bool Tick(float DeltaTime) override;
+	/**
+	 * Duration of samples to cache behind the play head (when not running as game).
+	 *
+	 * @see CacheAhead, CacheBehindGame
+	 */
+	UPROPERTY(BlueprintReadWrite, Category=Caching)
+	FTimespan CacheBehind;
 
-public:
+	/**
+	 * Duration of samples to cache behind the play head (when running as game).
+	 *
+	 * @see CacheAhead, CacheBehind
+	 */
+	UPROPERTY(BlueprintReadWrite, Category=Caching)
+	FTimespan CacheBehindGame;
 
-	/** Automatically start playback after media opened successfully. */
-	UPROPERTY(EditAnywhere, Category=Playback)
-	bool PlayOnOpen;
+public:
 
-	/** Whether playback should shuffle media sources in the play list. */
-	UPROPERTY(BlueprintReadOnly, EditAnywhere, Category=Playback)
-	uint32 Shuffle : 1;
+	/**
+	 * Output any audio via the operating system's sound mixer instead of a Sound Wave asset.
+	 *
+	 * If enabled, the assigned Sound Wave asset will be ignored. The SetNativeVolume
+	 * function can then be used to change the audio output volume at runtime. Note that
+	 * not all media player plug-ins may support native audio output on all platforms.
+	 *
+	 * @see SetNativeVolume
+	 */
+	UPROPERTY(BlueprintReadWrite, Category=Output, AdvancedDisplay)
+	bool NativeAudioOut;
 
 public:
 
-	/** Get an event delegate that is invoked when a media event occurred. */
-	DECLARE_EVENT_OneParam(UMediaPlayer, FOnMediaEvent, EMediaEvent /*Event*/)
-	FOnMediaEvent& OnMediaEvent()
-	{
-		return MediaEvent;
-	}
+	/**
+	 * Automatically start playback after media opened successfully.
+	 *
+	 * If disabled, listen to the OnMediaOpened Blueprint event to detect when
+	 * the media finished opening, and then start playback using the Play function.
+	 *
+	 * @see OpenFile, OpenPlaylist, OpenPlaylistIndex, OpenSource, OpenUrl, Play
+	 */
+	UPROPERTY(BlueprintReadWrite, EditAnywhere, Category=Playback)
+	bool PlayOnOpen;
+
+	/**
+	 * Whether playback should shuffle media sources in the play list.
+	 *
+	 * @see OpenPlaylist, OpenPlaylistIndex
+	 */
+	UPROPERTY(BlueprintReadWrite, EditAnywhere, Category=Playback)
+	uint32 Shuffle : 1;
 
 protected:
 
-	/** Whether the player should loop when media playback reaches the end. */
+	/**
+	 * Whether the player should loop when media playback reaches the end.
+	 *
+	 * Use the SetLooping function to change this value at runtime.
+	 *
+	 * @see IsLooping, SetLooping
+	 */
 	UPROPERTY(BlueprintReadOnly, EditAnywhere, Category=Playback)
 	uint32 Loop:1;
 
-	/** The overlay asset to output captions and subtitles to. */
-	UPROPERTY(BlueprintReadOnly, EditAnywhere, Category=Output)
-	UMediaOverlays* Overlays;
-
-	/** The play list to use, if any. */
+	/**
+	 * The play list to use, if any.
+	 *
+	 * Use the OpenPlaylist or OpenPlaylistIndex function to change this value at runtime.
+	 *
+	 * @see OpenPlaylist, OpenPlaylistIndex
+	 */
 	UPROPERTY(BlueprintReadOnly, transient, Category=Playback)
 	UMediaPlaylist* Playlist;
 
-	/** The current index of the source in the play list being played. */
+	/**
+	 * The current index of the source in the play list being played.
+	 *
+	 * Use the Previous and Next methods to change this value at runtime.
+	 *
+	 * @see Next, Previous
+	 */
 	UPROPERTY(BlueprintReadOnly, Category=Playback)
 	int32 PlaylistIndex;
 
-	/** The media sound wave to output the audio track samples to. */
-	UPROPERTY(BlueprintReadOnly, EditAnywhere, Category=Output)
-	UMediaSoundWave* SoundWave;
-
-	/** The media texture to output the video track frames to. */
-	UPROPERTY(BlueprintReadOnly, EditAnywhere, Category=Output)
-	UMediaTexture* VideoTexture;
+protected:
 
-private:
+	/**
+	 * The initial horizontal field of view (in Euler degrees; default = 90).
+	 *
+	 * This setting is used only for 360 videos. It determines the portion of the
+	 * video that is visible at a time. To modify the field of view at runtime in
+	 * Blueprints, use the SetHorizontalFieldOfView function.
+	 *
+	 * @see GetHorizontalFieldOfView, SetHorizontalFieldOfView, VerticalFieldOfView, ViewRotation
+	 */
+	UPROPERTY(EditAnywhere, Category=ViewSettings)
+	float HorizontalFieldOfView;
 
-	/** Callback for when a media overlays are being destroyed. */
-	void HandleMediaOverlaysBeginDestroy(UMediaOverlays& DestroyedOverlays);
+	/**
+	 * The initial vertical field of view (in Euler degrees; default = 60).
+	 *
+	 * This setting is used only for 360 videos. It determines the portion of the
+	 * video that is visible at a time. To modify the field of view at runtime in
+	 * Blueprints, use the SetHorizontalFieldOfView function.
+	 *
+	 * Please note that some 360 video players may be able to change only the
+	 * horizontal field of view, and this setting may be ignored.
+	 *
+	 * @see GetVerticalFieldOfView, SetVerticalFieldOfView, HorizontalFieldOfView, ViewRotation
+	 */
+	UPROPERTY(EditAnywhere, Category=ViewSettings)
+	float VerticalFieldOfView;
 
-	/** Callback for when a media texture is being destroyed. */
-	void HandleMediaSoundWaveBeginDestroy(UMediaSoundWave& DestroyedSoundWave);
+	/**
+	 * The initial view rotation.
+	 *
+	 * This setting is used only for 360 videos. It determines the rotation of
+	 * the video's view. To modify the view orientation at runtime in Blueprints,
+	 * use the GetViewRotation and SetViewRotation functions.
+	 *
+	 * Please note that not all players may support video view rotations.
+	 *
+	 * @see GetViewRotation, SetViewRotation, HorizontalFieldOfView, VerticalFieldOfView
+	 */
+	UPROPERTY(EditAnywhere, Category=ViewSettings)
+	FRotator ViewRotation;
 
-	/** Callback for when a media texture is being destroyed. */
-	void HandleMediaTextureBeginDestroy(UMediaTexture& DestroyedMediaTexture);
+private:
 
 	/** Callback for when a media event occurred in the player. */
 	void HandlePlayerMediaEvent(EMediaEvent Event);
 
 private:
 
-	/** The URL of the media that was last attempted to be opened. */
-	FString LastUrl;
-
 	/** An event delegate that is invoked when a media event occurred. */
 	FOnMediaEvent MediaEvent;
 
-	/** The player implementation. */
-	TSharedPtr<FMediaPlayerBase, ESPMode::ThreadSafe> Player;
+	/** The player facade. */
+	TSharedPtr<FMediaPlayerFacade, ESPMode::ThreadSafe> PlayerFacade;
+
+	/** The player's globally unique identifier. */
+	UPROPERTY()
+	FGuid PlayerGuid;
+
+	/** Automatically start playback of next item in play list. */
+	bool PlayOnNext;
 
 #if WITH_EDITOR
 	/** Whether the player was playing in PIE/SIE. */
diff --git a/Engine/Source/Runtime/MediaAssets/Public/MediaPlaylist.h b/Engine/Source/Runtime/MediaAssets/Public/MediaPlaylist.h
--- a/Engine/Source/Runtime/MediaAssets/Public/MediaPlaylist.h
+++ b/Engine/Source/Runtime/MediaAssets/Public/MediaPlaylist.h
@@ -19,19 +19,43 @@ class MEDIAASSETS_API UMediaPlaylist
 {
 	GENERATED_BODY()
 
+public:
+
+	/** Whether the play list should loop (default = true). */
+	UPROPERTY(BlueprintReadWrite, EditAnywhere, Category=Playback)
+	uint32 Loop:1;
+
 public:
 
 	/**
 	 * Add a media source to the play list.
 	 *
 	 * @param MediaSource The media source to append.
-	 * @see Insert, RemoveAll, Remove
+	 * @return true if the media source was added, false otherwise.
+	 * @see AddFile, AddUrl, Insert, RemoveAll, Remove, Replace
 	 */
 	UFUNCTION(BlueprintCallable, Category="Media|MediaPlaylist")
-	void Add(UMediaSource* MediaSource)
-	{
-		Items.Add(MediaSource);
-	}
+	bool Add(UMediaSource* MediaSource);
+
+	/**
+	 * Add a media file path to the play list.
+	 *
+	 * @param FilePath The file path to add.
+	 * @return true if the file was added, false otherwise.
+	 * @see Add, AddUrl, Insert, RemoveAll, Remove, Replace
+	 */
+	UFUNCTION(BlueprintCallable, Category="Media|MediaPlaylist")
+	bool AddFile(const FString& FilePath);
+
+	/**
+	 * Add a media URL to the play list.
+	 *
+	 * @param Url The URL to add.
+	 * @return true if the URL was added, false otherwise.
+	 * @see Add, AddFile, Insert, RemoveAll, Remove, Replace
+	 */
+	UFUNCTION(BlueprintCallable, Category="Media|MediaPlaylist")
+	bool AddUrl(const FString& Url);
 
 	/**
 	 * Get the media source at the specified index.
@@ -66,26 +90,22 @@ public:
 	/**
 	 * Get a random media source in the play list.
 	 *
-	 * @param InOutIndex Index of the current media source (will contain the new index).
+	 * @param OutIndex Will contain the index of the returned media source.
 	 * @return The random media source, or nullptr if the list is empty.
 	 * @see Get, GetNext, GetPrevious
 	 */
 	UFUNCTION(BlueprintCallable, Category="Media|MediaPlaylist")
-	UMediaSource* GetRandom(int32& InOutIndex);
+	UMediaSource* GetRandom(int32& OutIndex);
 
 	/**
 	 * Insert a media source into the play list at the given position.
 	 *
 	 * @param MediaSource The media source to insert.
 	 * @param Index The index to insert into.
-	 * @see Add, Remove, RemoveAll
+	 * @see Add, Remove, RemoveAll, Replace
 	 */
 	UFUNCTION(BlueprintCallable, Category="Media|MediaPlaylist")
-	void Insert(UMediaSource* MediaSource, int32 Index)
-	{
-		Index = FMath::Clamp(Index, 0, Items.Num());
-		Items.Insert(MediaSource, Index);
-	}
+	void Insert(UMediaSource* MediaSource, int32 Index);
 
 	/**
 	 * Get the number of media sources in the play list.
@@ -102,28 +122,32 @@ public:
 	 * Remove all occurrences of the given media source in the play list.
 	 *
 	 * @param MediaSource The media source to remove.
-	 * @see Add, Insert, Remove
+	 * @return true if the media source was removed, false otherwise.
+	 * @see Add, Insert, Remove, Replace
 	 */
 	UFUNCTION(BlueprintCallable, Category="Media|MediaPlaylist")
-	void Remove(UMediaSource* MediaSource)
-	{
-		Items.Remove(MediaSource);
-	}
+	bool Remove(UMediaSource* MediaSource);
 
 	/**
 	 * Remove the media source at the specified position.
 	 *
 	 * @param Index The index of the media source to remove.
-	 * @see Add, Insert, RemoveAll
+	 * @return true if the media source was removed, false otherwise.
+	 * @see Add, Insert, RemoveAll, Replace
 	 */
 	UFUNCTION(BlueprintCallable, Category="Media|MediaPlaylist")
-	void RemoveAt(int32 Index)
-	{
-		if (Items.IsValidIndex(Index))
-		{
-			Items.RemoveAt(Index);
-		}
-	}
+	bool RemoveAt(int32 Index);
+
+	/**
+	 * Replace the media source at the specified position.
+	 *
+	 * @param Index The index of the media source to replace.
+	 * @param Replacement The replacement media source.
+	 * @return true if the media source was replaced, false otherwise.
+	 * @see Add, Insert, RemoveAll, RemoveAt
+	 */
+	UFUNCTION(BlueprintCallable, Category="Media|MediaPlaylist")
+	bool Replace(int32 Index, UMediaSource* Replacement);
 
 protected:
 
