Media: Ensure facade's player guid matches its media player object after load

diff --git a/Engine/Source/Runtime/MediaAssets/Private/Assets/MediaPlayer.cpp b/Engine/Source/Runtime/MediaAssets/Private/Assets/MediaPlayer.cpp
--- a/Engine/Source/Runtime/MediaAssets/Private/Assets/MediaPlayer.cpp
+++ b/Engine/Source/Runtime/MediaAssets/Private/Assets/MediaPlayer.cpp
@@ -614,6 +614,7 @@ void UMediaPlayer::PostInitProperties()
 {
 	Super::PostInitProperties();
 
+	// Set the player GUID - required for UMediaPlayers dynamically allocated at runtime
 	PlayerFacade->SetGuid(PlayerGuid);
 
 	IMediaModule* MediaModule = FModuleManager::LoadModulePtr<IMediaModule>("Media");
@@ -625,6 +626,13 @@ void UMediaPlayer::PostInitProperties()
 	}
 }
 
+void UMediaPlayer::PostLoad()
+{
+	Super::PostLoad();
+
+	// Set the player GUID - required for UMediaPlayer assets
+	PlayerFacade->SetGuid(PlayerGuid);
+}
 
 #if WITH_EDITOR
 
diff --git a/Engine/Source/Runtime/MediaAssets/Public/MediaPlayer.h b/Engine/Source/Runtime/MediaAssets/Public/MediaPlayer.h
--- a/Engine/Source/Runtime/MediaAssets/Public/MediaPlayer.h
+++ b/Engine/Source/Runtime/MediaAssets/Public/MediaPlayer.h
@@ -880,6 +880,7 @@ public:
 	virtual FString GetDesc() override;
 	virtual void PostDuplicate(bool bDuplicateForPIE) override;
 	virtual void PostInitProperties() override;
+	virtual void PostLoad() override;
 
 #if WITH_EDITOR
 	virtual void PostEditChangeProperty(FPropertyChangedEvent& PropertyChangedEvent) override;
