Copying //UE4/Dev-Sequencer to //UE4/Dev-Main (Source: //UE4/Dev-Sequencer @ 3800792)

diff --git a/Engine/Source/Editor/Sequencer/Private/SequencerEdMode.cpp b/Engine/Source/Editor/Sequencer/Private/SequencerEdMode.cpp
--- a/Engine/Source/Editor/Sequencer/Private/SequencerEdMode.cpp
+++ b/Engine/Source/Editor/Sequencer/Private/SequencerEdMode.cpp
@@ -288,20 +288,16 @@ FTransform FSequencerEdMode::GetRefFrame(const TSharedPtr<FSequencer>& Sequencer
 									FVector ParentKeyPos;
 									FRotator ParentKeyRot;
 
-									const FMovieSceneEvaluationTemplateInstance* TemplateInstance = Sequencer->GetEvaluationTemplate().GetInstance(Sequencer->GetFocusedTemplateID());
-									if (TemplateInstance)
+									FMovieSceneEvaluationTemplate* Template = Sequencer->GetEvaluationTemplate().FindTemplate(Sequencer->GetFocusedTemplateID());
+									FMovieSceneEvaluationTrack* EvalTrack = Template ? Template->FindTrack(TransformTrack->GetSignature()) : nullptr;
+
+									if (EvalTrack)
 									{
-										for (FMovieSceneTrackIdentifier TrackID : TemplateInstance->Template->FindTracks(TransformTrack->GetSignature()))
-										{
-											if (const FMovieSceneEvaluationTrack* EvalTrack = TemplateInstance->Template->FindTrack(TrackID))
-											{
-												GetLocationAtTime(EvalTrack, ParentObject, KeyTime, ParentKeyPos, ParentKeyRot, Sequencer);
+										GetLocationAtTime(EvalTrack, ParentObject, KeyTime, ParentKeyPos, ParentKeyRot, Sequencer);
 
-												CurrentRefTM = FTransform(ParentKeyRot, ParentKeyPos);
+										CurrentRefTM = FTransform(ParentKeyRot, ParentKeyPos);
 
-												return CurrentRefTM * ParentRefTM;
-											}
-										}
+										return CurrentRefTM * ParentRefTM;
 									}
 								}
 							}
@@ -317,7 +313,7 @@ FTransform FSequencerEdMode::GetRefFrame(const TSharedPtr<FSequencer>& Sequencer
 	return RefTM;
 }
 
-void FSequencerEdMode::GetLocationAtTime(const FMovieSceneEvaluationTrack* Track, UObject* Object, float KeyTime, FVector& KeyPos, FRotator& KeyRot, const TSharedPtr<FSequencer>& Sequencer)
+void FSequencerEdMode::GetLocationAtTime(FMovieSceneEvaluationTrack* Track, UObject* Object, float KeyTime, FVector& KeyPos, FRotator& KeyRot, const TSharedPtr<FSequencer>& Sequencer)
 {
 	FMovieSceneInterrogationData InterrogationData;
 	Sequencer->GetEvaluationTemplate().CopyActuators(InterrogationData.GetAccumulator());
@@ -373,8 +369,8 @@ void FSequencerEdMode::DrawTransformTrack(const TSharedPtr<FSequencer>& Sequence
 		}
 	);
 	
-	const FMovieSceneEvaluationTemplateInstance* TemplateInstance = Sequencer->GetEvaluationTemplate().GetInstance(Sequencer->GetFocusedTemplateID());
-	if (!bShowTrajectory || !TemplateInstance || !TransformTrack->GetAllSections().ContainsByPredicate([](UMovieSceneSection* In){ return In->IsActive(); }))
+	FMovieSceneEvaluationTemplate* Template = Sequencer->GetEvaluationTemplate().FindTemplate(Sequencer->GetFocusedTemplateID());
+	if (!bShowTrajectory || !Template || !TransformTrack->GetAllSections().ContainsByPredicate([](UMovieSceneSection* In){ return In->IsActive(); }))
 	{
 		return;
 	}
@@ -382,14 +378,8 @@ void FSequencerEdMode::DrawTransformTrack(const TSharedPtr<FSequencer>& Sequence
 	FLinearColor TrackColor = TransformTrack->GetColorTint();
 
 	// Draw one line per-track (should only really ever be one)
-	for (FMovieSceneTrackIdentifier TrackID : TemplateInstance->Template->FindTracks(TransformTrack->GetSignature()))
+	if (FMovieSceneEvaluationTrack* EvalTrack = Template->FindTrack(TransformTrack->GetSignature()))
 	{
-		const FMovieSceneEvaluationTrack* EvalTrack = TemplateInstance->Template->FindTrack(TrackID);
-		if (!EvalTrack)
-		{
-			continue;
-		}
-
 		TArray<FTrajectoryKey> TrajectoryKeys = TransformTrack->GetTrajectoryData(Sequencer->GetLocalTime(), Sequencer->GetSettings()->GetTrajectoryPathCap());
 		for (TWeakObjectPtr<> WeakBinding : BoundObjects)
 		{
diff --git a/Engine/Source/Editor/Sequencer/Private/SequencerEdMode.h b/Engine/Source/Editor/Sequencer/Private/SequencerEdMode.h
--- a/Engine/Source/Editor/Sequencer/Private/SequencerEdMode.h
+++ b/Engine/Source/Editor/Sequencer/Private/SequencerEdMode.h
@@ -73,7 +73,7 @@ protected:
 	void DrawTransformTrack(const TSharedPtr<FSequencer>& Sequencer, FPrimitiveDrawInterface* PDI, UMovieScene3DTransformTrack* TransformTrack, const TArray<TWeakObjectPtr<UObject>>& BoundObjects, const bool bIsSelected);
 
 protected:
-	static void GetLocationAtTime(const FMovieSceneEvaluationTrack* Track, UObject* BoundObject, float KeyTime, FVector& KeyPos, FRotator& KeyRot, const TSharedPtr<FSequencer>& Sequencer);
+	static void GetLocationAtTime(FMovieSceneEvaluationTrack* Track, UObject* BoundObject, float KeyTime, FVector& KeyPos, FRotator& KeyRot, const TSharedPtr<FSequencer>& Sequencer);
 	static FTransform GetRefFrame(const TSharedPtr<FSequencer>& Sequencer, const UObject* InObject, float KeyTime);
 	static FTransform GetRefFrame(const TSharedPtr<FSequencer>& Sequencer, const AActor* Actor, float KeyTime);
 	static FTransform GetRefFrame(const TSharedPtr<FSequencer>& Sequencer, const USceneComponent* SceneComponent, float KeyTime);
diff --git a/Engine/Source/Editor/Sequencer/Private/SequencerNodeTree.cpp b/Engine/Source/Editor/Sequencer/Private/SequencerNodeTree.cpp
--- a/Engine/Source/Editor/Sequencer/Private/SequencerNodeTree.cpp
+++ b/Engine/Source/Editor/Sequencer/Private/SequencerNodeTree.cpp
@@ -134,10 +134,20 @@ void FSequencerNodeTree::Update()
 	TArray<TSharedRef<FSequencerObjectBindingNode>> ObjectNodes;
 	for( const FMovieSceneBinding& Binding : Bindings )
 	{
+		if (!Sequencer.IsBindingVisible(Binding))
+		{
+			continue;
+		}
+
 		TSharedRef<FSequencerObjectBindingNode> ObjectBindingNode = AddObjectBinding( Binding.GetName(), Binding.GetObjectGuid(), GuidToBindingMap, ObjectNodes );
 
 		for( UMovieSceneTrack* Track : Binding.GetTracks() )
 		{
+			if (!Sequencer.IsTrackVisible(Track))
+			{
+				continue;
+			}
+
 			// Create the new track node
 			TSharedRef<FSequencerTrackNode> TrackNode = MakeShared<FSequencerTrackNode>(*Track, *FindOrAddTypeEditor(*Track), false, nullptr, *this);
 
@@ -470,8 +480,27 @@ bool FSequencerNodeTree::GetSavedExpansionState(const FSequencerDisplayNode& Nod
 bool FSequencerNodeTree::GetDefaultExpansionState( const FSequencerDisplayNode& Node ) const
 {
 	// Object nodes, and track nodes that are parent tracks are expanded by default.
-	return Node.GetType() == ESequencerNode::Object ||
-		(Node.GetType() == ESequencerNode::Track && static_cast<const FSequencerTrackNode&>(Node).GetSubTrackMode() == FSequencerTrackNode::ESubTrackMode::ParentTrack);
+	if (Node.GetType() == ESequencerNode::Object)
+	{
+		return true;
+	}
+
+	else if (Node.GetType() == ESequencerNode::Track)
+	{
+		const FSequencerTrackNode& TrackNode = static_cast<const FSequencerTrackNode&>(Node);
+
+		if (TrackNode.GetSubTrackMode() == FSequencerTrackNode::ESubTrackMode::ParentTrack)
+		{
+			return true;
+		}
+
+		if (TrackNode.GetTrackEditor().GetDefaultExpansionState(TrackNode.GetTrack()))
+		{
+			return true;
+		}
+	}
+
+	return false;
 }
 
 
diff --git a/Engine/Source/Editor/Sequencer/Private/SequencerObjectChangeListener.cpp b/Engine/Source/Editor/Sequencer/Private/SequencerObjectChangeListener.cpp
--- a/Engine/Source/Editor/Sequencer/Private/SequencerObjectChangeListener.cpp
+++ b/Engine/Source/Editor/Sequencer/Private/SequencerObjectChangeListener.cpp
@@ -340,17 +340,24 @@ void FSequencerObjectChangeListener::OnObjectPreEditChange( UObject* Object, con
 	}
 
 	// Call add key/track before the property changes so that pre-animated state can be saved off.
-	for( FEditPropertyChain::TIterator It(PropertyChain.GetHead()); It; ++It )
-	{
-		UProperty* Prop = *It;
+	FPropertyPath PropertyPath;
+	UProperty* ActiveNode = PropertyChain.GetActiveNode()->GetValue();
+	UProperty* ActiveMemberNode = PropertyChain.GetActiveMemberNode()->GetValue();
 
-		TArray<UObject*> Objects;
-		Objects.Add(Object);
+	if (ActiveMemberNode)
+	{
+		PropertyPath.AddProperty(FPropertyInfo(ActiveMemberNode));
+	}
 
-		FPropertyPath PropertyPath;
-		PropertyPath.AddProperty(FPropertyInfo(Prop));
-		BroadcastPropertyChanged(FKeyPropertyParams(Objects, PropertyPath, ESequencerKeyMode::AutoKey));
+	if (ActiveMemberNode != ActiveNode)
+	{
+		PropertyPath.AddProperty(FPropertyInfo(ActiveNode));
 	}
+
+	TArray<UObject*> Objects;
+	Objects.Add(Object);
+
+	BroadcastPropertyChanged(FKeyPropertyParams(Objects, PropertyPath, ESequencerKeyMode::AutoKey));
 }
 
 void FSequencerObjectChangeListener::OnObjectPostEditChange( UObject* Object, FPropertyChangedEvent& PropertyChangedEvent )
diff --git a/Engine/Source/Editor/Sequencer/Private/SequencerSettings.cpp b/Engine/Source/Editor/Sequencer/Private/SequencerSettings.cpp
--- a/Engine/Source/Editor/Sequencer/Private/SequencerSettings.cpp
+++ b/Engine/Source/Editor/Sequencer/Private/SequencerSettings.cpp
@@ -31,6 +31,7 @@ USequencerSettings::USequencerSettings( const FObjectInitializer& ObjectInitiali
 	CurveValueSnapInterval = 10.0f;
 	bSnapCurveValueToInterval = true;
 	bLabelBrowserVisible = false;
+	bShowSelectedNodesOnly = false;
 	bRewindOnRecord = true;
 	ZoomPosition = ESequencerZoomPosition::SZP_CurrentTime;
 	bAutoScrollEnabled = false;
@@ -417,6 +418,22 @@ void USequencerSettings::SetLabelBrowserVisible(bool Visible)
 	}
 }
 
+bool USequencerSettings::GetShowSelectedNodesOnly() const
+{
+	return bShowSelectedNodesOnly;
+}
+
+void USequencerSettings::SetShowSelectedNodesOnly(bool Visible)
+{
+	if (bShowSelectedNodesOnly != Visible)
+	{
+		bShowSelectedNodesOnly = Visible;
+		SaveConfig();
+
+		OnShowSelectedNodesOnlyChangedEvent.Broadcast();
+	}
+}
+
 bool USequencerSettings::ShouldRewindOnRecord() const
 {
 	return bRewindOnRecord;
diff --git a/Engine/Source/Editor/Sequencer/Private/Tools/EditToolDragOperations.cpp b/Engine/Source/Editor/Sequencer/Private/Tools/EditToolDragOperations.cpp
--- a/Engine/Source/Editor/Sequencer/Private/Tools/EditToolDragOperations.cpp
+++ b/Engine/Source/Editor/Sequencer/Private/Tools/EditToolDragOperations.cpp
@@ -372,6 +372,19 @@ void FMoveSection::OnBeginDrag(const FPointerEvent& MouseEvent, FVector2D LocalM
 		Section->GetKeyHandles(DraggedKeyHandles, Section->GetRange());
 		RelativeOffsets.Add(FRelativeOffset{ Section->GetStartTime() - InitialPosition.X, Section->GetEndTime() - InitialPosition.X });
 	}
+
+	TSet<UMovieSceneTrack*> Tracks;
+	for (int32 Index = 0; Index < Sections.Num(); ++Index)
+	{
+		Tracks.Add(Sections[Index].TrackNode->GetTrack());
+	}
+	for (UMovieSceneTrack* Track : Tracks)
+	{
+		for (auto Section : Track->GetAllSections())
+		{
+			InitialRowIndices.Add(FInitialRowIndex{ Section, Section->GetRowIndex() });
+		}
+	}
 }
 
 void FMoveSection::OnEndDrag(const FPointerEvent& MouseEvent, FVector2D LocalMousePos, const FVirtualTrackArea& VirtualTrackArea)
@@ -382,6 +395,7 @@ void FMoveSection::OnEndDrag(const FPointerEvent& MouseEvent, FVector2D LocalMou
 	}
 
 	DraggedKeyHandles.Empty();
+	InitialRowIndices.Empty();
 
 	TSet<UMovieSceneTrack*> Tracks;
 	bool bRowIndicesFixed = false;
@@ -564,6 +578,12 @@ void FMoveSection::OnDrag(const FPointerEvent& MouseEvent, FVector2D LocalMouseP
 					const int32 NewIndex = FMath::FloorToInt((VirtualMousePos.Y - VirtualSectionTop) / VirtualRowHeight);
 					TargetRowIndex = FMath::Clamp(NewIndex, 0, MaxRowIndex);
 				}
+				
+				// If close to the top of the row, move else everything down
+				if (VirtualMousePos.Y <= VirtualSectionTop)
+				{
+					TargetRowIndex = -1;
+				}
 			}
 			else if(Handle.TrackNode->GetSubTrackMode() == FSequencerTrackNode::ESubTrackMode::SubTrack)
 			{
@@ -608,9 +628,45 @@ void FMoveSection::OnDrag(const FPointerEvent& MouseEvent, FVector2D LocalMouseP
 			)
 		)
 		{
-			Section->Modify();
-			Section->SetRowIndex(TargetRowIndex);
-			bRowIndexChanged = true;
+			// Reached the top, move everything else we're not moving downwards
+			if (TargetRowIndex == -1)
+			{
+				if (!bSectionsAreOnDifferentRows)
+				{
+					// If the sections being moved are all at the top, and all others are below it, do nothing
+					bool bSectionsBeingMovedAreAtTop = true;
+					for (auto InitialRowIndex : InitialRowIndices)
+					{
+						if (!SectionsBeingMoved.Contains(InitialRowIndex.Section))
+						{
+							if (InitialRowIndex.RowIndex <= FirstRowIndex)
+							{
+								bSectionsBeingMovedAreAtTop = false;
+								break;
+							}
+						}
+					}
+
+					if (!bSectionsBeingMovedAreAtTop)
+					{
+						for (auto InitialRowIndex : InitialRowIndices)
+						{
+							if (!SectionsBeingMoved.Contains(InitialRowIndex.Section))
+							{
+								InitialRowIndex.Section->Modify();
+								InitialRowIndex.Section->SetRowIndex(InitialRowIndex.RowIndex+1);
+								bRowIndexChanged = true;
+							}
+						}
+					}
+				}
+			}
+			else
+			{
+				Section->Modify();
+				Section->SetRowIndex(TargetRowIndex);
+				bRowIndexChanged = true;
+			}
 		}
 	}
 
diff --git a/Engine/Source/Editor/Sequencer/Private/Tools/EditToolDragOperations.h b/Engine/Source/Editor/Sequencer/Private/Tools/EditToolDragOperations.h
--- a/Engine/Source/Editor/Sequencer/Private/Tools/EditToolDragOperations.h
+++ b/Engine/Source/Editor/Sequencer/Private/Tools/EditToolDragOperations.h
@@ -139,6 +139,9 @@ private:
 
 	/** A handle for the sequencer node tree updated delegate. */
 	FDelegateHandle SequencerNodeTreeUpdatedHandle;
+
+	struct FInitialRowIndex { UMovieSceneSection* Section; int32 RowIndex; };
+	TArray<FInitialRowIndex> InitialRowIndices;
 };
 
 
diff --git a/Engine/Source/Editor/Sequencer/Private/Tools/SequencerEditTool_Movement.cpp b/Engine/Source/Editor/Sequencer/Private/Tools/SequencerEditTool_Movement.cpp
--- a/Engine/Source/Editor/Sequencer/Private/Tools/SequencerEditTool_Movement.cpp
+++ b/Engine/Source/Editor/Sequencer/Private/Tools/SequencerEditTool_Movement.cpp
@@ -71,6 +71,10 @@ FReply FSequencerEditTool_Movement::OnMouseMove(SWidget& OwnerWidget, const FGeo
 			if (DragOperation.IsValid())
 			{
 				DragPosition = MyGeometry.AbsoluteToLocal(MouseEvent.GetScreenSpacePosition());
+				
+				float CurrentTime = VirtualTrackArea.PixelToTime(DragPosition.X);
+				Sequencer.UpdateAutoScroll(CurrentTime);
+				
 				DragOperation->OnDrag(MouseEvent, DragPosition, VirtualTrackArea);
 			}
 		}
@@ -264,7 +268,10 @@ FReply FSequencerEditTool_Movement::OnMouseButtonUp(SWidget& OwnerWidget, const
 				Menu->GetOnMenuDismissed().AddLambda(
 					[=](TSharedRef<IMenu>)
 					{
-						ExistingHotspot->bIsLocked = false;
+						if (ExistingHotspot.IsValid())
+						{
+							ExistingHotspot->bIsLocked = false;
+						}
 						if (SequencerPtr->GetHotspot() == ExistingHotspot)
 						{
 							SequencerPtr->SetHotspot(nullptr);
diff --git a/Engine/Source/Editor/Sequencer/Public/ISequencer.h b/Engine/Source/Editor/Sequencer/Public/ISequencer.h
--- a/Engine/Source/Editor/Sequencer/Public/ISequencer.h
+++ b/Engine/Source/Editor/Sequencer/Public/ISequencer.h
@@ -11,6 +11,7 @@
 #include "Containers/ArrayView.h"
 #include "IMovieScenePlayer.h"
 #include "KeyPropertyParams.h"
+#include "MovieSceneBinding.h"
 #include "Widgets/Input/NumericTypeInterface.h"
 #include "Editor/SequencerWidgets/Public/ITimeSlider.h"
 
@@ -22,6 +23,7 @@ class FUICommandList;
 class ISequencerKeyCollection;
 class UMovieSceneSequence;
 class UMovieSceneSubSection;
+class IDetailsView;
 enum class EMapChangeType : uint8;
 
 /**
@@ -127,16 +129,19 @@ class ISequencer
 public:
 	
 	DECLARE_MULTICAST_DELEGATE(FOnGlobalTimeChanged);
+	DECLARE_MULTICAST_DELEGATE(FOnPlayEvent);
+	DECLARE_MULTICAST_DELEGATE(FOnStopEvent);
 	DECLARE_MULTICAST_DELEGATE(FOnBeginScrubbingEvent);
 	DECLARE_MULTICAST_DELEGATE(FOnEndScrubbingEvent);
 	DECLARE_MULTICAST_DELEGATE_OneParam(FOnMovieSceneDataChanged, EMovieSceneDataChangeType);
 	DECLARE_MULTICAST_DELEGATE(FOnMovieSceneBindingsChanged);
+	DECLARE_MULTICAST_DELEGATE_OneParam(FOnMovieSceneBindingsPasted, const TArray<FMovieSceneBinding>&);
 	
-	DECLARE_MULTICAST_DELEGATE_OneParam(FOnSelectionChangedObjectGuids, TArray<FGuid> /*Object*/)
+	DECLARE_MULTICAST_DELEGATE_OneParam(FOnSelectionChangedObjectGuids, TArray<FGuid> /*Object*/);
 
-	DECLARE_MULTICAST_DELEGATE_OneParam(FOnSelectionChangedTracks, TArray<UMovieSceneTrack*> /*Tracks*/)
+	DECLARE_MULTICAST_DELEGATE_OneParam(FOnSelectionChangedTracks, TArray<UMovieSceneTrack*> /*Tracks*/);
 
-	DECLARE_MULTICAST_DELEGATE_OneParam(FOnSelectionChangedSections, TArray<UMovieSceneSection*> /*Sections*/)
+	DECLARE_MULTICAST_DELEGATE_OneParam(FOnSelectionChangedSections, TArray<UMovieSceneSection*> /*Sections*/);
 
 	DECLARE_MULTICAST_DELEGATE_TwoParams(FOnActorAddedToSequencer, AActor*, const FGuid);
 
@@ -182,9 +187,10 @@ public:
 	 * Attempts to add a new spawnable to the MovieScene for the specified object (asset, class or actor instance)
 	 *
 	 * @param	Object	The asset, class, or actor to add a spawnable for
+	 * @param	ActorFactory	Optional actor factory to use to create spawnable type
 	 * @return	The spawnable guid for the spawnable, or an invalid Guid if we were not able to create a spawnable
 	 */
-	virtual FGuid MakeNewSpawnable(UObject& SourceObject) = 0;
+	virtual FGuid MakeNewSpawnable(UObject& SourceObject, UActorFactory* ActorFactory = nullptr) = 0;
 
 	/**
 	 * Add actors as possessable objects to sequencer.
@@ -349,6 +355,17 @@ public:
 	DECLARE_EVENT_OneParam(ISequencer, FOnActivateSequence, FMovieSceneSequenceIDRef)
 	virtual FOnActivateSequence& OnActivateSequence() = 0;
 
+	DECLARE_EVENT_TwoParams(ISequencer, FOnInitializeDetailsPanel, TSharedRef<IDetailsView>, TSharedRef<ISequencer>)
+	FOnInitializeDetailsPanel& OnInitializeDetailsPanel() { return InitializeDetailsPanelEvent; }
+
+	/** A delegate which will determine whether a binding should be visible in the tree. */
+	DECLARE_DELEGATE_RetVal_OneParam(bool, FOnGetIsBindingVisible, const FMovieSceneBinding&)
+	FOnGetIsBindingVisible& OnGetIsBindingVisible() { return GetIsBindingVisible; }
+
+	/** A delegate which will determine whether a track should be visible in the tree. */
+	DECLARE_DELEGATE_RetVal_OneParam(bool, FOnGetIsTrackVisible, const UMovieSceneTrack*)
+	FOnGetIsTrackVisible& OnGetIsTrackVisible() { return GetIsTrackVisible; }
+
 	/**
 	 * Gets a handle to runtime information about the object being manipulated by a movie scene
 	 * 
@@ -370,16 +387,20 @@ public:
 	DEPRECATED( 4.13, "NotifyMovieSceneDataChanged() is deprecated, use the version that takes EMovieSceneDataChangeType" )
 	void NotifyMovieSceneDataChanged() { NotifyMovieSceneDataChangedInternal(); };
 
+	/** Refresh the sequencer tree view */
+	virtual void RefreshTree() = 0;
+
 protected:
 	virtual void NotifyMovieSceneDataChangedInternal() = 0;
 
 public:
 	virtual void NotifyMovieSceneDataChanged( EMovieSceneDataChangeType DataChangeType ) = 0;
 
-	virtual void UpdateRuntimeInstances() = 0;
-
 	virtual void UpdatePlaybackRange() = 0;
 
+	virtual void SetPlaybackSpeed(float InPlaybackSpeed) = 0;
+	virtual float GetPlaybackSpeed() const = 0;
+
 	/** Get all the keys for the current sequencer selection */
 	virtual void GetKeysFromSelection(TUniquePtr<ISequencerKeyCollection>& KeyCollection, float DuplicateThresoldTime) = 0;
 
@@ -410,6 +431,12 @@ public:
 	/** Gets a multicast delegate which is executed whenever the global time changes. */
 	virtual FOnGlobalTimeChanged& OnGlobalTimeChanged() = 0;
 
+	/** Gets a multicast delegate which is executed whenever the user begins playing the sequence. */
+	virtual FOnPlayEvent& OnPlayEvent() = 0;
+
+	/** Gets a multicast delegate which is executed whenever the user stops playing the sequence. */
+	virtual FOnStopEvent& OnStopEvent() = 0;
+
 	/** Gets a multicast delegate which is executed whenever the user begins scrubbing. */
 	virtual FOnBeginScrubbingEvent& OnBeginScrubbingEvent() = 0;
 
@@ -422,6 +449,9 @@ public:
 	/** Gets a multicast delegate which is executed whenever the movie scene bindings are changed. */
 	virtual FOnMovieSceneBindingsChanged& OnMovieSceneBindingsChanged() = 0;
 
+	/** Gets a multicast delegate which is executed whenever bindings are pasted. */
+	virtual FOnMovieSceneBindingsPasted& OnMovieSceneBindingsPasted() = 0;
+
 	/** Gets a multicast delegate with an array of FGuid of bound objects which is called when the outliner node selection changes. */
 	virtual FOnSelectionChangedObjectGuids& GetSelectionChangedObjectGuids() = 0;
 
@@ -444,7 +474,7 @@ public:
 	virtual TSharedRef<SWidget> MakeTransportControls(bool bExtended) = 0;
 
 	/** Play or toggle playback at the specified play rate */
-	virtual FReply OnPlay(bool bTogglePlay = true, float InPlayRate = 1.f) = 0;
+	virtual FReply OnPlay(bool bTogglePlay = true) = 0;
 
 	/** Pause playback */
 	virtual void Pause() = 0;
@@ -475,4 +505,9 @@ public:
 	 * @return the widget
 	 */
 	virtual TSharedPtr<class ITimeSlider> GetTopTimeSliderWidget() const = 0;
+
+protected:
+	FOnInitializeDetailsPanel InitializeDetailsPanelEvent;
+	FOnGetIsBindingVisible GetIsBindingVisible;
+	FOnGetIsTrackVisible GetIsTrackVisible;
 };
diff --git a/Engine/Source/Editor/Sequencer/Public/ISequencerTrackEditor.h b/Engine/Source/Editor/Sequencer/Public/ISequencerTrackEditor.h
--- a/Engine/Source/Editor/Sequencer/Public/ISequencerTrackEditor.h
+++ b/Engine/Source/Editor/Sequencer/Public/ISequencerTrackEditor.h
@@ -185,6 +185,14 @@ public:
 		
 	}
 
+	/**
+	 * @return The default expansion state of this track
+	 */
+	virtual bool GetDefaultExpansionState(UMovieSceneTrack* InTrack) const
+	{
+		return false;
+	}
+
 public:
 
 	/** Virtual destructor. */
diff --git a/Engine/Source/Editor/Sequencer/Public/LevelEditorSequencerIntegration.h b/Engine/Source/Editor/Sequencer/Public/LevelEditorSequencerIntegration.h
--- a/Engine/Source/Editor/Sequencer/Public/LevelEditorSequencerIntegration.h
+++ b/Engine/Source/Editor/Sequencer/Public/LevelEditorSequencerIntegration.h
@@ -142,11 +142,14 @@ private:
 	/** Called when allow edits mode has changed */
 	void OnAllowEditsModeChanged(EAllowEditsMode AllowEditsMode);
 
-	/** Called when the user begins scrubbing */
-	void OnBeginScrubbing();
+	/** Called when the user begins playing/scrubbing */
+	void OnBeginDeferUpdates();
 
-	/** Called when the user stops scrubbing */
-	void OnEndScrubbing();
+	/** Called when the user stops playing/scrubbing */
+	void OnEndDeferUpdates();
+
+	/** Called to determine whether a binding is visible in the tree view */
+	bool IsBindingVisible(const FMovieSceneBinding& InBinding);
 
 	void OnPropertyEditorOpened();
 
@@ -210,5 +213,5 @@ private:
 
 	TSharedPtr<class FDetailKeyframeHandlerWrapper> KeyFrameHandler;
 
-	bool bScrubbing;
+	bool bDeferUpdates;
 };
diff --git a/Engine/Source/Editor/Sequencer/Public/MovieSceneTrackEditor.h b/Engine/Source/Editor/Sequencer/Public/MovieSceneTrackEditor.h
--- a/Engine/Source/Editor/Sequencer/Public/MovieSceneTrackEditor.h
+++ b/Engine/Source/Editor/Sequencer/Public/MovieSceneTrackEditor.h
@@ -26,7 +26,16 @@ struct FKeyPropertyResult
 	FKeyPropertyResult()
 		: bTrackModified(false)
 		, bHandleCreated(false)
-		, bTrackCreated(false) {}
+		, bTrackCreated(false)
+		, bKeyCreated(false) {}
+
+	inline void operator |= (const FKeyPropertyResult& A)
+	{
+		bTrackModified |= A.bTrackModified;
+		bHandleCreated |= A.bHandleCreated;
+		bTrackCreated |= A.bTrackCreated;
+		bKeyCreated |= A.bKeyCreated;
+	}
 
 	/* Was the track modified in any way? */
 	bool bTrackModified;
@@ -36,6 +45,9 @@ struct FKeyPropertyResult
 
 	/* Was a track created? */
 	bool bTrackCreated;
+
+	/* Was a key created? */
+	bool bKeyCreated;
 };
 
 /** Delegate for adding keys for a property
diff --git a/Engine/Source/Editor/Sequencer/Public/SequencerCommands.h b/Engine/Source/Editor/Sequencer/Public/SequencerCommands.h
--- a/Engine/Source/Editor/Sequencer/Public/SequencerCommands.h
+++ b/Engine/Source/Editor/Sequencer/Public/SequencerCommands.h
@@ -205,6 +205,9 @@ public:
 	/** Toggle the visibility of the transform box. */
 	TSharedPtr< FUICommandInfo > ToggleShowTransformBox;
 
+	/** Bake transform. */
+	TSharedPtr< FUICommandInfo > BakeTransform;
+
 	/** Turns the range slider on and off. */
 	TSharedPtr< FUICommandInfo > ToggleShowRangeSlider;
 
@@ -253,6 +256,9 @@ public:
 	/** Turns auto scroll on and off. */
 	TSharedPtr< FUICommandInfo > ToggleAutoScroll;
 
+	/** Toggles whether or not to show selected nodes only. */
+	TSharedPtr< FUICommandInfo > ToggleShowSelectedNodesOnly;
+
 	/** Toggles whether or not the curve editor should be shown. */
 	TSharedPtr< FUICommandInfo > ToggleShowCurveEditor;
 
@@ -310,6 +316,9 @@ public:
 	/** Exports animation to fbx. */
 	TSharedPtr< FUICommandInfo > ExportFBX;
 
+	/** Exports animation to camera anim. */
+	TSharedPtr< FUICommandInfo > ExportToCameraAnim;
+
 	/** Toggle whether we should evaluate sub sequences in isolation */
 	TSharedPtr< FUICommandInfo > ToggleEvaluateSubSequencesInIsolation;
 
diff --git a/Engine/Source/Editor/Sequencer/Public/SequencerSettings.h b/Engine/Source/Editor/Sequencer/Public/SequencerSettings.h
--- a/Engine/Source/Editor/Sequencer/Public/SequencerSettings.h
+++ b/Engine/Source/Editor/Sequencer/Public/SequencerSettings.h
@@ -110,6 +110,7 @@ public:
 	GENERATED_UCLASS_BODY()
 
 	DECLARE_MULTICAST_DELEGATE( FOnEvaluateSubSequencesInIsolationChanged );
+	DECLARE_MULTICAST_DELEGATE( FOnShowSelectedNodesOnlyChanged );
 	DECLARE_MULTICAST_DELEGATE_OneParam( FOnAllowEditsModeChanged, EAllowEditsMode );
 	DECLARE_MULTICAST_DELEGATE_OneParam( FOnLockPlaybackToAudioClockChanged, bool );
 
@@ -230,6 +231,12 @@ public:
 	/** Sets whether or not the label browser is visible. */
 	void SetLabelBrowserVisible(bool Visible);
 
+	/** Gets whether or not to show selected nodes only. */
+	bool GetShowSelectedNodesOnly() const;
+	/** Sets whether or not to show selected nodes only. */
+	void SetShowSelectedNodesOnly(bool Visible);
+	FOnShowSelectedNodesOnlyChanged& GetOnShowSelectedNodesOnlyChanged() { return OnShowSelectedNodesOnlyChangedEvent; }
+
 	/** Gets whether to jump to the start of the sequence when we start a recording or not. */
 	bool ShouldRewindOnRecord() const;
 	/** Sets whether to jump to the start of the sequence when we start a recording. */
@@ -447,6 +454,10 @@ protected:
 	UPROPERTY( config, EditAnywhere, Category=General )
 	bool bLabelBrowserVisible;
 
+	/** Only show selected nodes in the tree view. */
+	UPROPERTY( config, EditAnywhere, Category=General )
+	bool bShowSelectedNodesOnly;
+
 	/** Defines whether to jump back to the start of the sequence when a recording is started */
 	UPROPERTY(config, EditAnywhere, Category=General)
 	bool bRewindOnRecord;
@@ -537,5 +548,6 @@ protected:
 
 	FOnLockPlaybackToAudioClockChanged OnLockPlaybackToAudioClockChanged;
 	FOnEvaluateSubSequencesInIsolationChanged OnEvaluateSubSequencesInIsolationChangedEvent;
+	FOnShowSelectedNodesOnlyChanged OnShowSelectedNodesOnlyChangedEvent;
 	FOnAllowEditsModeChanged OnAllowEditsModeChangedEvent;
 };
diff --git a/Engine/Source/Editor/UMGEditor/Private/WidgetBlueprintEditor.cpp b/Engine/Source/Editor/UMGEditor/Private/WidgetBlueprintEditor.cpp
--- a/Engine/Source/Editor/UMGEditor/Private/WidgetBlueprintEditor.cpp
+++ b/Engine/Source/Editor/UMGEditor/Private/WidgetBlueprintEditor.cpp
@@ -87,6 +87,7 @@ FWidgetBlueprintEditor::~FWidgetBlueprintEditor()
 	if ( Sequencer.IsValid() )
 	{
 		Sequencer->OnMovieSceneDataChanged().RemoveAll( this );
+		Sequencer->OnMovieSceneBindingsPasted().RemoveAll( this );
 		Sequencer.Reset();
 	}
 
@@ -702,6 +703,7 @@ TSharedPtr<ISequencer>& FWidgetBlueprintEditor::GetSequencer()
 
 		Sequencer = FModuleManager::LoadModuleChecked<ISequencerModule>("Sequencer").CreateSequencer(SequencerInitParams);
 		Sequencer->OnMovieSceneDataChanged().AddSP( this, &FWidgetBlueprintEditor::OnMovieSceneDataChanged );
+		Sequencer->OnMovieSceneBindingsPasted().AddSP( this, &FWidgetBlueprintEditor::OnMovieSceneBindingsPasted );
 		// Change selected widgets in the sequencer tree view
 		Sequencer->GetSelectionChangedObjectGuids().AddSP(this, &FWidgetBlueprintEditor::SyncSelectedWidgetsWithSequencerSelection);
 		ChangeViewedAnimation(*UWidgetAnimation::GetNullAnimation());
@@ -1321,6 +1323,57 @@ void FWidgetBlueprintEditor::OnMovieSceneDataChanged(EMovieSceneDataChangeType D
 	bRefreshGeneratedClassAnimations = true;
 }
 
+void FWidgetBlueprintEditor::OnMovieSceneBindingsPasted(const TArray<FMovieSceneBinding>& BindingsPasted)
+{
+	TArray<FObjectAndDisplayName> BindableObjects;
+	{
+		GetBindableObjects(GetPreview()->WidgetTree, BindableObjects);
+	}
+
+	UMovieSceneSequence* AnimationSequence = GetSequencer().Get()->GetFocusedMovieSceneSequence();
+	UObject* BindingContext = GetAnimationPlaybackContext();
+
+	// First, rebind top level possessables (without parents) - match binding pasted's name with the bindable object name
+	for (const FMovieSceneBinding& BindingPasted : BindingsPasted)
+	{
+		FMovieScenePossessable* Possessable = AnimationSequence->GetMovieScene()->FindPossessable(BindingPasted.GetObjectGuid());
+		if (Possessable && !Possessable->GetParent().IsValid())
+		{
+			for (FObjectAndDisplayName& BindableObject : BindableObjects)
+			{
+				if (BindableObject.DisplayName.ToString() == BindingPasted.GetName())
+				{
+					AnimationSequence->BindPossessableObject(BindingPasted.GetObjectGuid(), *BindableObject.Object, BindingContext);			
+					break;
+				}
+			}
+		}
+	}
+
+	// Second, bind child possessables - match the binding pasted's parent guid with the bindable slot's content guid
+	for (const FMovieSceneBinding& BindingPasted : BindingsPasted)
+	{
+		FMovieScenePossessable* Possessable = AnimationSequence->GetMovieScene()->FindPossessable(BindingPasted.GetObjectGuid());
+		if (Possessable && Possessable->GetParent().IsValid())
+		{
+			for (FObjectAndDisplayName& BindableObject : BindableObjects)
+			{
+				UPanelSlot* PanelSlot = Cast<UPanelSlot>(BindableObject.Object);
+				if (PanelSlot && PanelSlot->Content)
+				{
+					FGuid ParentGuid = AnimationSequence->FindPossessableObjectId(*PanelSlot->Content, BindingContext);
+
+					if (ParentGuid == Possessable->GetParent())
+					{
+						AnimationSequence->BindPossessableObject(BindingPasted.GetObjectGuid(), *BindableObject.Object, BindingContext);			
+						break;
+					}
+				}
+			}
+		}
+	}
+}
+
 void FWidgetBlueprintEditor::SyncSelectedWidgetsWithSequencerSelection(TArray<FGuid> ObjectGuids)
 {
 	UMovieSceneSequence* AnimationSequence = GetSequencer().Get()->GetFocusedMovieSceneSequence();
diff --git a/Engine/Source/Editor/UnrealEd/Private/Fbx/FbxCurveDataImport.cpp b/Engine/Source/Editor/UnrealEd/Private/Fbx/FbxCurveDataImport.cpp
--- a/Engine/Source/Editor/UnrealEd/Private/Fbx/FbxCurveDataImport.cpp
+++ b/Engine/Source/Editor/UnrealEd/Private/Fbx/FbxCurveDataImport.cpp
@@ -269,20 +269,18 @@ namespace UnFbx {
 		{
 		case FbxAnimCurveDef::eInterpolationCubic:
 		{
-			switch (FbxKey.GetTangentMode())
+			FbxAnimCurveDef::ETangentMode TangentMode = FbxKey.GetTangentMode(true);
+			if (TangentMode & (FbxAnimCurveDef::eTangentUser | FbxAnimCurveDef::eTangentTCB | FbxAnimCurveDef::eTangentGenericClamp | FbxAnimCurveDef::eTangentGenericClampProgressive))
+			{
+				Mode = CIM_CurveUser;
+			}
+			else if (TangentMode & FbxAnimCurveDef::eTangentGenericBreak)
 			{
-			case FbxAnimCurveDef::eTangentBreak:
 				Mode = CIM_CurveBreak;
-				break;
-			case FbxAnimCurveDef::eTangentAuto:
+			}
+			else if (TangentMode & FbxAnimCurveDef::eTangentAuto)
+			{
 				Mode = CIM_CurveAuto;
-				break;
-			case FbxAnimCurveDef::eTangentUser:
-			case FbxAnimCurveDef::eTangentTCB:
-				Mode = CIM_CurveUser;
-				break;
-			default:
-				break;
 			}
 			break;
 		}
diff --git a/Engine/Source/Editor/UnrealEd/Private/Fbx/FbxMainExport.cpp b/Engine/Source/Editor/UnrealEd/Private/Fbx/FbxMainExport.cpp
--- a/Engine/Source/Editor/UnrealEd/Private/Fbx/FbxMainExport.cpp
+++ b/Engine/Source/Editor/UnrealEd/Private/Fbx/FbxMainExport.cpp
@@ -2521,22 +2521,31 @@ void FFbxExporter::ExportLevelSequenceFloatTrack( FbxNode& FbxActor, UMovieScene
 		Property = FbxActor.FindProperty( "UE_MotionBlur_Amount", false );
 	}
 
-	if ( Property != 0 )
+	FRichCurve& FloatCurve = FloatSection->GetFloatCurve();
+
+	if (Property == 0)
 	{
-		FRichCurve& FloatCurve = FloatSection->GetFloatCurve();
+		CreateAnimatableUserProperty(&FbxActor, FloatCurve.GetDefaultValue(), TCHAR_TO_UTF8(*PropertyName), TCHAR_TO_UTF8(*PropertyName));
 
-		FbxAnimCurve* AnimCurve = FbxAnimCurve::Create( Scene, "" );
-		FbxAnimCurveNode* CurveNode = Property.GetCurveNode( true );
-		if ( !CurveNode )
-		{
-			return;
-		}
+		Property = FbxActor.FindProperty(TCHAR_TO_UTF8(*PropertyName), false);
+	}
 
-		CurveNode->SetChannelValue<double>( 0U, FloatCurve.GetDefaultValue() );
-		CurveNode->ConnectToChannel( AnimCurve, 0U );
+	if (Property == 0)
+	{
+		return;
+	}
 
-		ExportRichCurveToFbxCurve(*AnimCurve, FloatCurve, IsFoV ? ERichCurveValueMode::Fov : ERichCurveValueMode::Default);
+	FbxAnimCurve* AnimCurve = FbxAnimCurve::Create( Scene, "" );
+	FbxAnimCurveNode* CurveNode = Property.GetCurveNode( true );
+	if ( !CurveNode )
+	{
+		return;
 	}
+
+	CurveNode->SetChannelValue<double>( 0U, FloatCurve.GetDefaultValue() );
+	CurveNode->ConnectToChannel( AnimCurve, 0U );
+
+	ExportRichCurveToFbxCurve(*AnimCurve, FloatCurve, IsFoV ? ERichCurveValueMode::Fov : ERichCurveValueMode::Default);
 }
 		
 /**
diff --git a/Engine/Source/Editor/UnrealEd/Private/SCurveEditor.cpp b/Engine/Source/Editor/UnrealEd/Private/SCurveEditor.cpp
--- a/Engine/Source/Editor/UnrealEd/Private/SCurveEditor.cpp
+++ b/Engine/Source/Editor/UnrealEd/Private/SCurveEditor.cpp
@@ -2322,7 +2322,7 @@ void SCurveEditor::MoveSelectedKeys(FVector2D Delta)
 	TArray<FRichCurveEditInfo> ChangedCurveEditInfos;
 
 	const FScopedTransaction Transaction( LOCTEXT("CurveEditor_MoveKeys", "Move Keys") );
-	CurveOwner->ModifyOwner();
+	CurveOwner->ModifyOwnerChange();
 
 	// track all unique curves encountered so their tangents can be updated later
 	TSet<FRichCurve*> UniqueCurves;
@@ -3518,6 +3518,7 @@ bool SCurveEditor::IsPostInfinityExtrapSelected(ERichCurveExtrapolation Extrapol
 void SCurveEditor::MoveTangents(FTrackScaleInfo& ScaleInfo, FVector2D Delta)
 {
 	TArray<FRichCurveEditInfo> ChangedCurveEditInfos;
+	CurveOwner->ModifyOwnerChange();
 
 	for (auto SelectedTangent : SelectedTangents)
 	{
diff --git a/Engine/Source/Editor/VREditor/VREditorActions.cpp b/Engine/Source/Editor/VREditor/VREditorActions.cpp
--- a/Engine/Source/Editor/VREditor/VREditorActions.cpp
+++ b/Engine/Source/Editor/VREditor/VREditorActions.cpp
@@ -432,7 +432,8 @@ void FVREditorActionCallbacks::PlaySequenceAtRate(UVREditorMode* InVRMode, float
 	ISequencer* CurrentSequencer = InVRMode->GetCurrentSequencer();
 	if (CurrentSequencer != nullptr)
 	{
-		CurrentSequencer->OnPlay(false, Rate);
+		CurrentSequencer->SetPlaybackSpeed(Rate);
+		CurrentSequencer->OnPlay(false);
 	}
 }
 
@@ -451,8 +452,8 @@ void FVREditorActionCallbacks::PlayFromBeginning(UVREditorMode* InVRMode)
 	if (CurrentSequencer != nullptr)
 	{
 		CurrentSequencer->SetLocalTime(0.0f);
-		const float Rate = 1.0f;
-		CurrentSequencer->OnPlay(false, Rate);
+		CurrentSequencer->SetPlaybackSpeed(1.f);
+		CurrentSequencer->OnPlay(false);
 	}
 }
 
diff --git a/Engine/Source/Programs/AutomationTool/IOS/IOSPlatform.Automation.cs b/Engine/Source/Programs/AutomationTool/IOS/IOSPlatform.Automation.cs
--- a/Engine/Source/Programs/AutomationTool/IOS/IOSPlatform.Automation.cs
+++ b/Engine/Source/Programs/AutomationTool/IOS/IOSPlatform.Automation.cs
@@ -974,7 +974,7 @@ public class IOSPlatform : Platform
 		}
 	}
 
-	private void StageMovieFiles(DirectoryReference InputDir, DeploymentContext SC)
+	protected void StageMovieFiles(DirectoryReference InputDir, DeploymentContext SC)
 	{
 		if(DirectoryReference.Exists(InputDir))
 		{
diff --git a/Engine/Source/Programs/AutomationTool/TVOS/TVOSPlatform.Automation.cs b/Engine/Source/Programs/AutomationTool/TVOS/TVOSPlatform.Automation.cs
--- a/Engine/Source/Programs/AutomationTool/TVOS/TVOSPlatform.Automation.cs
+++ b/Engine/Source/Programs/AutomationTool/TVOS/TVOSPlatform.Automation.cs
@@ -101,17 +101,8 @@ public class TVOSPlatform : IOSPlatform
 
         // copy the movies from the project
         {
-			DirectoryReference PlatformMovies = DirectoryReference.Combine(SC.ProjectRoot, "Build", "TVOS", "Resources", "Movies");
-			if(DirectoryReference.Exists(PlatformMovies))
-			{
-				SC.StageFiles(StagedFileType.SystemNonUFS, PlatformMovies, StageFilesSearch.TopDirectoryOnly, StagedDirectoryReference.Root);
-			}
-
-			DirectoryReference AllMovies = DirectoryReference.Combine(SC.ProjectRoot, "Content", "Movies");
-			if(DirectoryReference.Exists(AllMovies))
-			{
-				SC.StageFiles(StagedFileType.SystemNonUFS, AllMovies, StageFilesSearch.AllDirectories, StagedDirectoryReference.Root);
-			}
+            StageMovieFiles(DirectoryReference.Combine(SC.EngineRoot, "Content", "Movies"), SC);
+            StageMovieFiles(DirectoryReference.Combine(SC.ProjectRoot, "Content", "Movies"), SC);
         }
     }
 }
diff --git a/Engine/Source/Runtime/AudioMixer/Private/Components/SynthComponent.cpp b/Engine/Source/Runtime/AudioMixer/Private/Components/SynthComponent.cpp
--- a/Engine/Source/Runtime/AudioMixer/Private/Components/SynthComponent.cpp
+++ b/Engine/Source/Runtime/AudioMixer/Private/Components/SynthComponent.cpp
@@ -3,7 +3,6 @@
 #include "SynthComponent.h"
 #include "AudioDevice.h"
 #include "AudioMixerLog.h"
-#include "Components/BillboardComponent.h"
 
 USynthSound::USynthSound(const FObjectInitializer& ObjectInitializer)
 	: Super(ObjectInitializer)
@@ -11,12 +10,12 @@ USynthSound::USynthSound(const FObjectInitializer& ObjectInitializer)
 {
 }
 
-void USynthSound::Init(USynthComponent* InSynthComponent, const int32 InNumChannels, const int32 InSampleRate)
+void USynthSound::Init(USynthComponent* InSynthComponent, const int32 InNumChannels, const int32 InSampleRate, const int32 InCallbackSize)
 {
 	OwningSynthComponent = InSynthComponent;
 	bVirtualizeWhenSilent = true;
 	NumChannels = InNumChannels;
-
+	NumSamplesToGeneratePerCallback = InCallbackSize;
 	// Turn off async generation in old audio engine on mac.
 #if PLATFORM_MAC
 	if (!InSynthComponent->GetAudioDevice()->IsAudioMixerEnabled())
@@ -106,6 +105,8 @@ USynthComponent::USynthComponent(const FObjectInitializer& ObjectInitializer)
 	// Set the default sound class
 	SoundClass = USoundBase::DefaultSoundClassObject;
 
+	PreferredBufferLength = DEFAULT_PROCEDURAL_SOUNDWAVE_BUFFER_SIZE;
+
 #if WITH_EDITORONLY_DATA
 	bVisualizeComponent = false;
 #endif
@@ -193,7 +194,7 @@ void USynthComponent::Initialize(int32 SampleRateOverride)
 			Synth->SoundSubmixObject = SoundSubmix;
 			Synth->SoundSubmixSends = SoundSubmixSends;
 
-			Synth->Init(this, NumChannels, SampleRate);
+			Synth->Init(this, NumChannels, SampleRate, PreferredBufferLength);
 		}
 	}
 }
diff --git a/Engine/Source/Runtime/AudioMixer/Public/Components/SynthComponent.h b/Engine/Source/Runtime/AudioMixer/Public/Components/SynthComponent.h
--- a/Engine/Source/Runtime/AudioMixer/Public/Components/SynthComponent.h
+++ b/Engine/Source/Runtime/AudioMixer/Public/Components/SynthComponent.h
@@ -34,7 +34,7 @@ class AUDIOMIXER_API USynthSound : public USoundWaveProcedural
 {
 	GENERATED_UCLASS_BODY()
 
-	void Init(USynthComponent* InSynthComponent, const int32 InNumChannels, const int32 SampleRate);
+	void Init(USynthComponent* InSynthComponent, const int32 InNumChannels, const int32 SampleRate, const int32 InCallbackSize);
 	void StartOnAudioDevice(FAudioDevice* InAudioDevice);
 
 	/** Begin USoundWave */
@@ -199,6 +199,9 @@ protected:
 	// Can be set by the derived class, defaults to 2
 	int32 NumChannels;
 
+	// Can be set by the derived class- sets the preferred callback size for the synth component.
+	int32 PreferredBufferLength;
+
 private:
 
 	UPROPERTY(Transient)
diff --git a/Engine/Source/Runtime/CinematicCamera/Private/CameraRig_Rail.cpp b/Engine/Source/Runtime/CinematicCamera/Private/CameraRig_Rail.cpp
--- a/Engine/Source/Runtime/CinematicCamera/Private/CameraRig_Rail.cpp
+++ b/Engine/Source/Runtime/CinematicCamera/Private/CameraRig_Rail.cpp
@@ -111,9 +111,10 @@ void ACameraRig_Rail::UpdatePreviewMeshes()
 
 			for (int PtIdx = 0; PtIdx < NumSplinePoints - 1; ++PtIdx)
 			{
-				FVector StartLoc, StartTangent, EndLoc, EndTangent;
-				RailSplineComponent->GetLocationAndTangentAtSplinePoint(PtIdx, StartLoc, StartTangent, ESplineCoordinateSpace::Local);
-				RailSplineComponent->GetLocationAndTangentAtSplinePoint(PtIdx + 1, EndLoc, EndTangent, ESplineCoordinateSpace::Local);
+				FVector StartLoc = RailSplineComponent->GetLocationAtSplinePoint(PtIdx, ESplineCoordinateSpace::Local);
+				FVector StartTangent = RailSplineComponent->GetLeaveTangentAtSplinePoint(PtIdx, ESplineCoordinateSpace::Local);
+				FVector EndTangent = RailSplineComponent->GetArriveTangentAtSplinePoint(PtIdx + 1,ESplineCoordinateSpace::Local);
+				FVector EndLoc = RailSplineComponent->GetLocationAtSplinePoint(PtIdx + 1, ESplineCoordinateSpace::Local);
 
 				USplineMeshComponent* const SplineMeshComp = PreviewRailMeshSegments[PtIdx];
 				if (SplineMeshComp)
diff --git a/Engine/Source/Runtime/Core/Public/Math/Range.h b/Engine/Source/Runtime/Core/Public/Math/Range.h
--- a/Engine/Source/Runtime/Core/Public/Math/Range.h
+++ b/Engine/Source/Runtime/Core/Public/Math/Range.h
@@ -51,6 +51,9 @@ public:
 
 	typedef TRangeBound<ElementType> BoundsType;
 
+	/*~ Typedef used to pass/return small element types by value rather than const& */
+	typedef typename TCallTraits<ElementType>::ParamType ElementValueOrConstRef;
+
 	/** Default constructor (no initialization). */
 	TRange() { }
 
@@ -61,7 +64,7 @@ public:
 	 *
 	 * @param A The element in the range.
 	 */
-	explicit TRange(const ElementType& A)
+	explicit TRange(ElementValueOrConstRef A)
 		: LowerBound(BoundsType::Inclusive(A))
 		, UpperBound(BoundsType::Inclusive(A))
 	{ }
@@ -74,7 +77,7 @@ public:
 	 * @param A The range's lower bound value (inclusive).
 	 * @param B The range's upper bound value (exclusive).
 	 */
-	explicit TRange(const ElementType& A, const ElementType& B)
+	explicit TRange(ElementValueOrConstRef A, ElementValueOrConstRef B)
 		: LowerBound(BoundsType::Inclusive(A))
 		, UpperBound(BoundsType::Exclusive(B))
 	{ }
@@ -184,7 +187,7 @@ public:
 	 * @param Element The element to check.
 	 * @return true if the range contains the element, false otherwise.
 	 */
-	bool Contains(const ElementType& Element) const
+	bool Contains(ElementValueOrConstRef Element) const
 	{
 		return ((BoundsType::MinLower(LowerBound, Element) == LowerBound) &&
 				(BoundsType::MaxUpper(UpperBound, Element) == UpperBound));
@@ -234,7 +237,7 @@ public:
 	 * @return Bound value.
 	 * @see GetLowerBound, GetUpperBoundValue, HasLowerBound
 	 */
-	const ElementType& GetLowerBoundValue() const
+	ElementValueOrConstRef GetLowerBoundValue() const
 	{
 		return LowerBound.GetValue();
 	}
@@ -258,7 +261,7 @@ public:
 	 * @return Bound value.
 	 * @see GetLowerBoundValue, GetUpperBound, HasUpperBound
 	 */
-	const ElementType& GetUpperBoundValue() const
+	ElementValueOrConstRef GetUpperBoundValue() const
 	{
 		return UpperBound.GetValue();
 	}
@@ -390,7 +393,7 @@ public:
 	 *
 	 * @param Element The element at which to split the range.
 	 */
-	TArray<TRange> Split(const ElementType& Element) const
+	TArray<TRange> Split(ElementValueOrConstRef Element) const
 	{
 		TArray<TRange> Result;
 		
@@ -596,7 +599,7 @@ public:
 	 * @return A new range.
 	 * @see All, AtMost, Empty, Exclusive, GreaterThan, Inclusive, LessThan
 	 */
-	static FORCEINLINE TRange AtLeast(const ElementType& Value)
+	static FORCEINLINE TRange AtLeast(ElementValueOrConstRef Value)
 	{
 		return TRange(BoundsType::Inclusive(Value), BoundsType::Open());
 	}
@@ -608,7 +611,7 @@ public:
 	 * @return A new range.
 	 * @see All, AtLeast, Empty, Exclusive, GreaterThan, Inclusive, LessThan
 	 */
-	static FORCEINLINE TRange AtMost(const ElementType& Value)
+	static FORCEINLINE TRange AtMost(ElementValueOrConstRef Value)
 	{
 		return TRange(BoundsType::Open(), BoundsType::Inclusive(Value));
 	}
@@ -632,7 +635,7 @@ public:
 	 * @return A new range.
 	 * @see All, AtLeast, AtMost, Empty, Exclusive, GreaterThan, Inclusive, LessThan
 	 */
-	static FORCEINLINE TRange Exclusive(const ElementType& Min, const ElementType& Max)
+	static FORCEINLINE TRange Exclusive(ElementValueOrConstRef Min, ElementValueOrConstRef Max)
 	{
 		return TRange(BoundsType::Exclusive(Min), BoundsType::Exclusive(Max));
 	}
@@ -644,7 +647,7 @@ public:
 	 * @return A new range.
 	 * @see All, AtLeast, AtMost, Empty, Exclusive, Inclusive, LessThan
 	 */
-	static FORCEINLINE TRange GreaterThan(const ElementType& Value)
+	static FORCEINLINE TRange GreaterThan(ElementValueOrConstRef Value)
 	{
 		return TRange(BoundsType::Exclusive(Value), BoundsType::Open());
 	}
@@ -657,7 +660,7 @@ public:
 	 * @return A new range.
 	 * @see All, AtLeast, AtMost, Empty, Exclusive, GreaterThan, LessThan
 	 */
-	static FORCEINLINE TRange Inclusive(const ElementType& Min, const ElementType& Max)
+	static FORCEINLINE TRange Inclusive(ElementValueOrConstRef Min, ElementValueOrConstRef Max)
 	{
 		return TRange(BoundsType::Inclusive(Min), BoundsType::Inclusive(Max));
 	}
@@ -669,7 +672,7 @@ public:
 	 * @return A new range.
 	 * @see All, AtLeast, AtMost, Empty, Exclusive, GreaterThan, Inclusive
 	 */
-	static FORCEINLINE TRange LessThan(const ElementType& Value)
+	static FORCEINLINE TRange LessThan(ElementValueOrConstRef Value)
 	{
 		return TRange(BoundsType::Open(), BoundsType::Exclusive(Value));
 	}
@@ -729,12 +732,12 @@ private:
 		{ \
 		} \
 		 \
-		explicit Name(const ElementType& A) \
+		explicit Name(ElementValueOrConstRef A) \
 			: Super(A) \
 		{ \
 		} \
 		 \
-		explicit Name(const ElementType& A, const ElementType& B) \
+		explicit Name(ElementValueOrConstRef A, ElementValueOrConstRef B) \
 			: Super(A, B) \
 		{ \
 		} \
@@ -744,7 +747,7 @@ private:
 		{ \
 		} \
 		 \
-		TArray<Name> Split(const ElementType& Element) const \
+		TArray<Name> Split(ElementValueOrConstRef Element) const \
 		{ \
 			return TArray<Name>(Super::Split(Element)); \
 		} \
@@ -789,22 +792,22 @@ private:
 			return Super::All(); \
 		} \
 		 \
-		static FORCEINLINE Name AtLeast(const ElementType& Value) \
+		static FORCEINLINE Name AtLeast(ElementValueOrConstRef Value) \
 		{ \
 			return Super::AtLeast(Value); \
 		} \
 		 \
-		static FORCEINLINE Name AtMost(const ElementType& Value) \
+		static FORCEINLINE Name AtMost(ElementValueOrConstRef Value) \
 		{ \
 			return Super::AtMost(Value); \
 		} \
 		 \
-		static FORCEINLINE TRange GreaterThan(const ElementType& Value) \
+		static FORCEINLINE TRange GreaterThan(ElementValueOrConstRef Value) \
 		{ \
 			return Super::GreaterThan(Value); \
 		} \
 		 \
-		static FORCEINLINE TRange LessThan(const ElementType& Value) \
+		static FORCEINLINE TRange LessThan(ElementValueOrConstRef Value) \
 		{ \
 			return Super::LessThan(Value); \
 		} \
diff --git a/Engine/Source/Runtime/Core/Public/Math/RangeBound.h b/Engine/Source/Runtime/Core/Public/Math/RangeBound.h
--- a/Engine/Source/Runtime/Core/Public/Math/RangeBound.h
+++ b/Engine/Source/Runtime/Core/Public/Math/RangeBound.h
@@ -35,6 +35,9 @@ class TRangeBound
 {
 public:
 
+	/*~ Typedef used to pass/return small element types by value rather than const& */
+	typedef typename TCallTraits<ElementType>::ParamType ElementValueOrConstRef;
+
 	/**
 	 * Default constructor.
 	 *
@@ -51,7 +54,7 @@ public:
 	 * @param InValue The bound's value.
 	 * @see Exclusive, Inclusive, Open
 	 */
-	TRangeBound(const ElementType& InValue)
+	TRangeBound(ElementValueOrConstRef InValue)
 		: Type(ERangeBoundTypes::Inclusive)
 		, Value(InValue)
 	{ }
@@ -90,7 +93,7 @@ public:
 	 * @return Bound value.
 	 * @see IsOpen
 	 */
-	const ElementType& GetValue() const
+	FORCEINLINE_DEBUGGABLE ElementValueOrConstRef GetValue() const
 	{
 		check(Type != ERangeBoundTypes::Open);
 
@@ -102,7 +105,7 @@ public:
 	 *
 	 * @return true if the bound is closed, false otherwise.
 	 */
-	bool IsClosed() const
+	FORCEINLINE_DEBUGGABLE bool IsClosed() const
 	{
 		return (Type != ERangeBoundTypes::Open);
 	}
@@ -112,7 +115,7 @@ public:
 	 *
 	 * @return true if the bound is exclusive, false otherwise.
 	 */
-	bool IsExclusive() const
+	FORCEINLINE_DEBUGGABLE bool IsExclusive() const
 	{
 		return (Type == ERangeBoundTypes::Exclusive);
 	}
@@ -122,7 +125,7 @@ public:
 	 *
 	 * @return true if the bound is inclusive, false otherwise.
 	 */
-	bool IsInclusive() const
+	FORCEINLINE_DEBUGGABLE bool IsInclusive() const
 	{
 		return (Type == ERangeBoundTypes::Inclusive);
 	}
@@ -132,7 +135,7 @@ public:
 	 *
 	 * @return true if the bound is open, false otherwise.
 	 */
-	bool IsOpen() const
+	FORCEINLINE_DEBUGGABLE bool IsOpen() const
 	{
 		return (Type == ERangeBoundTypes::Open);
 	}
@@ -170,7 +173,7 @@ public:
 	 * @param Value The bound value.
 	 * @return An exclusive closed bound.
 	 */
-	static FORCEINLINE TRangeBound Exclusive(const ElementType& Value)
+	static FORCEINLINE TRangeBound Exclusive(ElementValueOrConstRef Value)
 	{
 		TRangeBound Result;
 
@@ -186,7 +189,7 @@ public:
 	 * @param Value The bound value.
 	 * @return An inclusive closed bound.
 	 */
-	static FORCEINLINE TRangeBound Inclusive(const ElementType& Value)
+	static FORCEINLINE TRangeBound Inclusive(ElementValueOrConstRef Value)
 	{
 		TRangeBound Result;
 
@@ -338,12 +341,12 @@ private:
 			: Super(InValue) \
 		{ } \
 		 \
-		static FORCEINLINE Name Exclusive(const ElementType& Value) \
+		static FORCEINLINE Name Exclusive(ElementValueOrConstRef Value) \
 		{ \
 			return static_cast<const Name&>(Super::Exclusive(Value)); \
 		} \
 		 \
-		static FORCEINLINE Name Inclusive(const ElementType& Value) \
+		static FORCEINLINE Name Inclusive(ElementValueOrConstRef Value) \
 		{ \
 			return static_cast<const Name&>(Super::Inclusive(Value)); \
 		} \
diff --git a/Engine/Source/Runtime/Core/Public/UObject/SequencerObjectVersion.h b/Engine/Source/Runtime/Core/Public/UObject/SequencerObjectVersion.h
--- a/Engine/Source/Runtime/Core/Public/UObject/SequencerObjectVersion.h
+++ b/Engine/Source/Runtime/Core/Public/UObject/SequencerObjectVersion.h
@@ -24,6 +24,9 @@ struct CORE_API FSequencerObjectVersion
 		// When finished now defaults to restore state
 		WhenFinishedDefaultsToRestoreState,
 
+		// EvaluationTree added
+		EvaluationTree,
+
 		// When finished now defaults to project default
 		WhenFinishedDefaultsToProjectDefault,
 
diff --git a/Engine/Source/Runtime/Engine/Classes/Camera/CameraComponent.h b/Engine/Source/Runtime/Engine/Classes/Camera/CameraComponent.h
--- a/Engine/Source/Runtime/Engine/Classes/Camera/CameraComponent.h
+++ b/Engine/Source/Runtime/Engine/Classes/Camera/CameraComponent.h
@@ -173,6 +173,8 @@ public:
 	void AddAdditiveOffset(FTransform const& Transform, float FOV);
 	/** Removes any additive offset. */
 	void ClearAdditiveOffset();
+	/** Get the additive offset */
+	void GetAdditiveOffset(FTransform& OutAdditiveOffset, float& OutAdditiveFOVOffset) const;
 
 	/** Stores a given PP and weight to be later applied when the final PP is computed. Used for things like in-editor camera animation preview. */
 	void AddExtraPostProcessBlend(FPostProcessSettings const& PPSettings, float PPBlendWeight);
diff --git a/Engine/Source/Runtime/Engine/Classes/Curves/CurveOwnerInterface.h b/Engine/Source/Runtime/Engine/Classes/Curves/CurveOwnerInterface.h
--- a/Engine/Source/Runtime/Engine/Classes/Curves/CurveOwnerInterface.h
+++ b/Engine/Source/Runtime/Engine/Classes/Curves/CurveOwnerInterface.h
@@ -24,6 +24,9 @@ public:
 	/** Called to modify the owner of the curve */
 	virtual void ModifyOwner() = 0;
 
+	/** Called to modify the owner of the curve during interaction/while being edited. */
+	virtual void ModifyOwnerChange() { ModifyOwner(); }
+
 	/** Returns the owner(s) of the curve */
 	virtual TArray<const UObject*> GetOwners() const = 0;
 
diff --git a/Engine/Source/Runtime/Engine/Classes/Sound/SoundWaveProcedural.h b/Engine/Source/Runtime/Engine/Classes/Sound/SoundWaveProcedural.h
--- a/Engine/Source/Runtime/Engine/Classes/Sound/SoundWaveProcedural.h
+++ b/Engine/Source/Runtime/Engine/Classes/Sound/SoundWaveProcedural.h
@@ -15,6 +15,8 @@
 #include "Sound/SoundWave.h"
 #include "SoundWaveProcedural.generated.h"
 
+#define DEFAULT_PROCEDURAL_SOUNDWAVE_BUFFER_SIZE 1024
+
 DECLARE_DELEGATE_TwoParams( FOnSoundWaveProceduralUnderflow, class USoundWaveProcedural*, int32 );
 
 UCLASS()
diff --git a/Engine/Source/Runtime/Engine/Private/Camera/CameraComponent.cpp b/Engine/Source/Runtime/Engine/Private/Camera/CameraComponent.cpp
--- a/Engine/Source/Runtime/Engine/Private/Camera/CameraComponent.cpp
+++ b/Engine/Source/Runtime/Engine/Private/Camera/CameraComponent.cpp
@@ -236,7 +236,7 @@ void UCameraComponent::Serialize(FArchive& Ar)
 
 void UCameraComponent::GetCameraView(float DeltaTime, FMinimalViewInfo& DesiredView)
 {
-	if (bLockToHmd && GEngine->XRSystem.IsValid() && GetWorld()->WorldType != EWorldType::Editor )
+	if (bLockToHmd && GEngine && GEngine->XRSystem.IsValid() && GetWorld() && GetWorld()->WorldType != EWorldType::Editor )
 	{
 		IXRTrackingSystem* XRSystem = GEngine->XRSystem.Get();
 
@@ -367,6 +367,11 @@ void UCameraComponent::ClearAdditiveOffset()
 #endif
 }
 
+void UCameraComponent::GetAdditiveOffset(FTransform& OutAdditiveOffset, float& OutAdditiveFOVOffset) const
+{
+	OutAdditiveOffset = AdditiveOffset;
+	OutAdditiveFOVOffset = AdditiveFOVOffset;
+}
 
 void UCameraComponent::AddExtraPostProcessBlend(FPostProcessSettings const& PPSettings, float PPBlendWeight)
 {
diff --git a/Engine/Source/Runtime/Engine/Private/ExternalTexture.cpp b/Engine/Source/Runtime/Engine/Private/ExternalTexture.cpp
--- a/Engine/Source/Runtime/Engine/Private/ExternalTexture.cpp
+++ b/Engine/Source/Runtime/Engine/Private/ExternalTexture.cpp
@@ -2,108 +2,121 @@
 
 #include "ExternalTexture.h"
 
+
+#define EXTERNALTEXTURE_TRACE_REGISTRY 0
+
+
 FExternalTextureRegistry* FExternalTextureRegistry::Singleton = nullptr;
 
+
 FExternalTextureRegistry& FExternalTextureRegistry::Get()
 {
 	check(IsInRenderingThread());
+
 	if (Singleton == nullptr)
 	{
 		Singleton = new FExternalTextureRegistry();
 	}
+
 	return *Singleton;
 }
 
-/* Register an external texture, its sampler state and coordinate scale/bias against a GUID */
+
 void FExternalTextureRegistry::RegisterExternalTexture(const FGuid& InGuid, FTextureRHIRef& InTextureRHI, FSamplerStateRHIRef& InSamplerStateRHI, const FLinearColor& InCoordinateScaleRotation, const FLinearColor& InCoordinateOffset)
 {
 	TextureEntries.Add(InGuid, FExternalTextureEntry(InTextureRHI, InSamplerStateRHI, InCoordinateScaleRotation, InCoordinateOffset));
+
 	for (const FMaterialRenderProxy* MaterialRenderProxy : ReferencingMaterialRenderProxies)
 	{
 		const_cast<FMaterialRenderProxy*>(MaterialRenderProxy)->CacheUniformExpressions();
 	}
 }
 
-/* Removes an external texture given a GUID */
+
 void FExternalTextureRegistry::UnregisterExternalTexture(const FGuid& InGuid)
 {
 	TextureEntries.Remove(InGuid);
+
 	for (const FMaterialRenderProxy* MaterialRenderProxy : ReferencingMaterialRenderProxies)
 	{
 		const_cast<FMaterialRenderProxy*>(MaterialRenderProxy)->CacheUniformExpressions();
 	}
 }
 
-/* Removes the specified MaterialRenderProxy from the list of those using an external texture. */
+
 void FExternalTextureRegistry::RemoveMaterialRenderProxyReference(const FMaterialRenderProxy* MaterialRenderProxy)
 {
 	ReferencingMaterialRenderProxies.Remove(MaterialRenderProxy);
 }
 
-/* Registers the MaterialRenderProxy as using an external texture and looks up an external texture for given the specified GUID provided it is valid.
- * @return false if the texture is not registered or the GUID is invalid. 
- */
+
 bool FExternalTextureRegistry::GetExternalTexture(const FMaterialRenderProxy* MaterialRenderProxy, const FGuid& InGuid, FTextureRHIRef& OutTextureRHI, FSamplerStateRHIRef& OutSamplerStateRHI)
 {
-//	FPlatformMisc::LowLevelOutputDebugStringf(TEXT("GetExternalTexture: Guid = %s"), *InGuid.ToString());
+#if EXTERNALTEXTURE_TRACE_REGISTRY
+	FPlatformMisc::LowLevelOutputDebugStringf(TEXT("GetExternalTexture: Guid = %s"), *InGuid.ToString());
+#endif
 
-	// Only cache render proxies that have been initialized, since FMaterialRenderProxy::ReleaseDynamicRHI() is responsible for removing them
-	if (MaterialRenderProxy && MaterialRenderProxy->IsInitialized())
+	// register material proxy if already initialized
+	if ((MaterialRenderProxy != nullptr) && MaterialRenderProxy->IsInitialized())
 	{
 		ReferencingMaterialRenderProxies.Add(MaterialRenderProxy);
+
+		// Note: FMaterialRenderProxy::ReleaseDynamicRHI()
+		// is responsible for removing the material proxy
+	}
+
+	if (!InGuid.IsValid())
+	{
+		return false; // no identifier associated with the texture yet
 	}
 
-	// If the GUID was not valid, we still need to register the MaterialRenderProxy so the uniform expressions can be recached when the external texture gets assigned a valid GUID.
+	FExternalTextureEntry* Entry = TextureEntries.Find(InGuid);
 
-	if (InGuid.IsValid())
+	if (Entry == nullptr)
 	{
-		FExternalTextureEntry* Entry = TextureEntries.Find(InGuid);
-		if (Entry)
-		{
-//			FPlatformMisc::LowLevelOutputDebugStringf(TEXT("GetExternalTexture: Found"));
-			OutTextureRHI = Entry->TextureRHI;
-			OutSamplerStateRHI = Entry->SamplerStateRHI;
-			return true;
-		}
-		else
-		{
-//			FPlatformMisc::LowLevelOutputDebugStringf(TEXT("GetExternalTexture: NOT FOUND!"));
-		}
+#if EXTERNALTEXTURE_TRACE_REGISTRY
+		FPlatformMisc::LowLevelOutputDebugStringf(TEXT("GetExternalTexture: NOT FOUND!"));
+#endif
+
+		return false; // texture not registered
 	}
 
-	return false;
+#if EXTERNALTEXTURE_TRACE_REGISTRY
+	FPlatformMisc::LowLevelOutputDebugStringf(TEXT("GetExternalTexture: Found"));
+#endif
+
+	OutTextureRHI = Entry->TextureRHI;
+	OutSamplerStateRHI = Entry->SamplerStateRHI;
+
+	return true;
 }
 
-/* Looks up an texture coordinate scale rotation for given a given GUID
- * @return false if the texture is not registered
- */
+
 bool FExternalTextureRegistry::GetExternalTextureCoordinateScaleRotation(const FGuid& InGuid, FLinearColor& OutCoordinateScaleRotation)
 {
 	FExternalTextureEntry* Entry = TextureEntries.Find(InGuid);
-	if (Entry)
-	{
-		OutCoordinateScaleRotation = Entry->CoordinateScaleRotation;
-		return true;
-	}
-	else
+
+	if (Entry == nullptr)
 	{
 		return false;
 	}
+
+	OutCoordinateScaleRotation = Entry->CoordinateScaleRotation;
+
+	return true;
 }
 
-/* Looks up an texture coordinate offset for given a given GUID
-* @return false if the texture is not registered
-*/
+
 bool FExternalTextureRegistry::GetExternalTextureCoordinateOffset(const FGuid& InGuid, FLinearColor& OutCoordinateOffset)
 {
 	FExternalTextureEntry* Entry = TextureEntries.Find(InGuid);
-	if (Entry)
-	{
-		OutCoordinateOffset = Entry->CoordinateOffset;
-		return true;
-	}
-	else
+
+	if (Entry == nullptr)
 	{
 		return false;
 	}
+
+	OutCoordinateOffset = Entry->CoordinateOffset;
+
+	return true;
 }
diff --git a/Engine/Source/Runtime/Engine/Private/Materials/MaterialUniformExpressions.cpp b/Engine/Source/Runtime/Engine/Private/Materials/MaterialUniformExpressions.cpp
--- a/Engine/Source/Runtime/Engine/Private/Materials/MaterialUniformExpressions.cpp
+++ b/Engine/Source/Runtime/Engine/Private/Materials/MaterialUniformExpressions.cpp
@@ -5,12 +5,12 @@
 =============================================================================*/
 
 #include "Materials/MaterialUniformExpressions.h"
+#include "CoreGlobals.h"
 #include "SceneManagement.h"
 #include "Materials/MaterialInstance.h"
 #include "Materials/MaterialInstanceSupport.h"
 #include "Materials/MaterialParameterCollection.h"
 #include "ExternalTexture.h"
-#include "MessageLog.h"
 #include "UObjectToken.h"
 
 TLinkedList<FMaterialUniformExpressionType*>*& FMaterialUniformExpressionType::GetTypeList()
@@ -409,9 +409,7 @@ FUniformBufferRHIRef FUniformExpressionSet::CreateUniformBuffer(const FMaterialR
 						ExpressionIndex,
 						FText::FromString(*MaterialRenderContext.Material.GetFriendlyName()));
 
-					FMessageLog("PIE").Warning()
-						->AddToken(FUObjectToken::Create(Value))
-						->AddToken(FTextToken::Create(MessageText));
+					GLog->Logf(ELogVerbosity::Warning, *MessageText.ToString());
 				}
 			}
 
diff --git a/Engine/Source/Runtime/Engine/Private/SoundWaveProcedural.cpp b/Engine/Source/Runtime/Engine/Private/SoundWaveProcedural.cpp
--- a/Engine/Source/Runtime/Engine/Private/SoundWaveProcedural.cpp
+++ b/Engine/Source/Runtime/Engine/Private/SoundWaveProcedural.cpp
@@ -2,13 +2,36 @@
 
 #include "Sound/SoundWaveProcedural.h"
 
+#include "AudioDevice.h"
+#include "Engine/Engine.h"
+
+
 USoundWaveProcedural::USoundWaveProcedural(const FObjectInitializer& ObjectInitializer)
 	: Super(ObjectInitializer)
 {
 	bProcedural = true;
 	bReset = false;
 	NumBufferUnderrunSamples = 512;
-	NumSamplesToGeneratePerCallback = 1024;
+	NumSamplesToGeneratePerCallback = DEFAULT_PROCEDURAL_SOUNDWAVE_BUFFER_SIZE;
+
+	// If the main audio device has been set up, we can use this to define our callback size.
+	// We need to do this for procedural sound waves that we do not process asynchronously,
+	// to ensure that we do not underrun.
+	
+	if (GEngine)
+	{
+		FAudioDevice* MainAudioDevice = GEngine->GetMainAudioDevice();
+		if (MainAudioDevice && !MainAudioDevice->IsAudioMixerEnabled())
+		{
+#if PLATFORM_MAC
+			// We special case the mac callback on the old audio engine, Since Buffer Length is smaller than the device callback size.
+			NumSamplesToGeneratePerCallback = 2048;
+#else
+			NumSamplesToGeneratePerCallback = MainAudioDevice->GetBufferLength();
+#endif
+			NumBufferUnderrunSamples = NumSamplesToGeneratePerCallback / 2;
+		}
+	}
 
 	SampleByteSize = 2;
 
diff --git a/Engine/Source/Runtime/Engine/Public/ExternalTexture.h b/Engine/Source/Runtime/Engine/Public/ExternalTexture.h
--- a/Engine/Source/Runtime/Engine/Public/ExternalTexture.h
+++ b/Engine/Source/Runtime/Engine/Public/ExternalTexture.h
@@ -7,15 +7,22 @@
 #include "RHIResources.h"
 #include "MaterialShared.h"
 
-/* FExternalTextureRegistry - stores a registry of external textures mapped to their GUIDs */
-
+/**
+ * Stores a registry of external textures mapped to their GUIDs.
+ */
 class ENGINE_API FExternalTextureRegistry
 {
+	/**
+	 * Registry singleton instance.
+	 *
+	 * @see Get
+	 */
 	static FExternalTextureRegistry* Singleton;
 
-	FExternalTextureRegistry()
-	{}
+	/** Default constructor. */
+	FExternalTextureRegistry() { }
 
+	/** An entry in the registry. */
 	struct FExternalTextureEntry
 	{
 		FExternalTextureEntry(FTextureRHIRef& InTextureRHI, FSamplerStateRHIRef& InSamplerStateRHI, const FLinearColor& InCoordinateScaleRotation, const FLinearColor& InCoordinateOffset)
@@ -31,34 +38,80 @@ class ENGINE_API FExternalTextureRegistry
 		FLinearColor CoordinateOffset;
 	};
 
+	/** Maps unique identifiers to texture entries. */
 	TMap<FGuid, FExternalTextureEntry> TextureEntries;
+
+	/** Set of registered material render proxies. */
 	TSet<const FMaterialRenderProxy*> ReferencingMaterialRenderProxies;
 
 public:
 
-	static FExternalTextureRegistry& Get();
-
-	/* Register an external texture, its sampler state and coordinate scale/bias against a GUID */
+	/**
+	 * Register an external texture, its sampler state and coordinate scale/bias against a GUID.
+	 *
+	 * @param InGuid The texture's unique identifier.
+	 * @param InTextureRHI The texture.
+	 * @param InSamplerStateRHI The texture's sampler state.
+	 * @param InCoordinateScaleRotation Texture coordinate scale and rotation parameters (optional; packed into color value).
+	 * @param InCoordinateOffset Texture coordinate offset (optional; packed into color value).
+	 * @see GetExternalTexture, UnregisterExternalTexture
+	 */
 	void RegisterExternalTexture(const FGuid& InGuid, FTextureRHIRef& InTextureRHI, FSamplerStateRHIRef& InSamplerStateRHI, const FLinearColor& InCoordinateScaleRotation = FLinearColor(1, 0, 0, 1), const FLinearColor& InCoordinateOffset = FLinearColor(0, 0, 0, 0));
 
-	/* Removes an external texture given a GUID */
+	/**
+	 * Removes an external texture given a GUID.
+	 *
+	 * @param InGuid The unique identifier of the texture to unregister.
+	 * @see RegisterExternalTexture
+	 */
 	void UnregisterExternalTexture(const FGuid& InGuid);
 
-	/* Removes the specified MaterialRenderProxy from the list of those using an external texture. */
+	/**
+	 * Removes the specified MaterialRenderProxy from the list of those using an external texture.
+	 *
+	 * @param MaterialRenderProxy The material render proxy to remove.
+	 * @see GetExternalTexture, RegisterExternalTexture, UnregisterExternalTexture
+	 */
 	void RemoveMaterialRenderProxyReference(const FMaterialRenderProxy* MaterialRenderProxy);
 
-	/* Registers the MaterialRenderProxy as using an external texture and looks up an external texture for given the specified GUID provided it is valid.
-	 * @return false if the texture is not registered or the GUID is invalid.
+	/**
+	 * Get the external texture with the specified identifier.
+	 *
+	 * @param MaterialRenderProxy The material render proxy that is using the texture (will be registered).
+	 * @param InGuid The texture's unique identifier.
+	 * @param OutTextureRHI Will contain the external texture.
+	 * @param OutSamplerStateRHI Will contain the texture's sampler state.
+	 * @return true on success, false if the identifier was not found.
+	 * @see RegisterExternalTexture, UnregisterExternalTexture
 	 */
 	bool GetExternalTexture(const FMaterialRenderProxy* MaterialRenderProxy, const FGuid& InGuid, FTextureRHIRef& OutTextureRHI, FSamplerStateRHIRef& OutSamplerStateRHI);
 
-	/* Looks up an texture coordinate scale rotation for given a given GUID
-	 * @return false if the texture is not registered
+	/**
+	 * Looks up an external texture's coordinate scale rotation.
+	 *
+	 * @param InGuid The texture's unique identifier.
+	 * @param OutCoordinateScaleRotation Will contain the texture's coordinate scale and rotation (packed into color value).
+	 * @return true on success, false if the identifier was not found.
+	 * @see GetExternalTextureCoordinateOffset
 	 */
 	bool GetExternalTextureCoordinateScaleRotation(const FGuid& InGuid, FLinearColor& OutCoordinateScaleRotation);
 
-	/* Looks up an texture coordinate offset for given a given GUID
-	 * @return false if the texture is not registered
+	/**
+	 * Looks up an texture coordinate offset for given a given GUID.
+	 *
+	 * @param InGuid The texture's unique identifier.
+	 * @param OutCoordinateOffset Will contain the texture's coordinate offset (packed into color value).
+	 * @return true on success, false if the identifier was not found.
+	 * @see GetExternalTextureCoordinateScaleRotation
 	 */
 	bool GetExternalTextureCoordinateOffset(const FGuid& InGuid, FLinearColor& OutCoordinateOffset);
+
+public:
+
+	/**
+	 * Get the registry singleton instance.
+	 *
+	 * @return External texture registry.
+	 */
+	static FExternalTextureRegistry& Get();
 };
diff --git a/Engine/Source/Runtime/LevelSequence/Private/LevelSequenceActor.cpp b/Engine/Source/Runtime/LevelSequence/Private/LevelSequenceActor.cpp
--- a/Engine/Source/Runtime/LevelSequence/Private/LevelSequenceActor.cpp
+++ b/Engine/Source/Runtime/LevelSequence/Private/LevelSequenceActor.cpp
@@ -5,6 +5,7 @@
 #include "Engine/Texture2D.h"
 #include "Components/BillboardComponent.h"
 #include "LevelSequenceBurnIn.h"
+#include "DefaultLevelSequenceInstanceData.h"
 
 #if WITH_EDITOR
 	#include "PropertyCustomizationHelpers.h"
@@ -45,6 +46,10 @@ ALevelSequenceActor::ALevelSequenceActor(const FObjectInitializer& Init)
 
 	BindingOverrides = Init.CreateDefaultSubobject<UMovieSceneBindingOverrides>(this, "BindingOverrides");
 	BurnInOptions = Init.CreateDefaultSubobject<ULevelSequenceBurnInOptions>(this, "BurnInOptions");
+	DefaultInstanceData = Init.CreateDefaultSubobject<UDefaultLevelSequenceInstanceData>(this, "InstanceData");
+
+	bOverrideInstanceData = false;
+
 	PrimaryActorTick.bCanEverTick = true;
 	bAutoPlay = false;
 }
@@ -214,10 +219,16 @@ void ALevelSequenceActor::InitializePlayer()
 
 	if (GetWorld()->IsGameWorld() && (LevelSequenceAsset != nullptr))
 	{
-		PlaybackSettings.BindingOverrides = BindingOverrides;
+		FMovieSceneSequencePlaybackSettings PlaybackSettingsCopy = PlaybackSettings;
+
+		PlaybackSettingsCopy.BindingOverrides = BindingOverrides;
+		if (bOverrideInstanceData)
+		{
+			PlaybackSettingsCopy.InstanceData = DefaultInstanceData;
+		}
 
 		SequencePlayer = NewObject<ULevelSequencePlayer>(this, "AnimationPlayer");
-		SequencePlayer->Initialize(LevelSequenceAsset, GetWorld(), PlaybackSettings);
+		SequencePlayer->Initialize(LevelSequenceAsset, GetWorld(), PlaybackSettingsCopy);
 		SequencePlayer->SetEventReceivers(TArray<UObject*>(AdditionalEventReceivers));
 
 		RefreshBurnIn();
diff --git a/Engine/Source/Runtime/LevelSequence/Private/LevelSequenceActorSpawner.cpp b/Engine/Source/Runtime/LevelSequence/Private/LevelSequenceActorSpawner.cpp
--- a/Engine/Source/Runtime/LevelSequence/Private/LevelSequenceActorSpawner.cpp
+++ b/Engine/Source/Runtime/LevelSequence/Private/LevelSequenceActorSpawner.cpp
@@ -53,6 +53,7 @@ UObject* FLevelSequenceActorSpawner::SpawnObject(FMovieSceneSpawnable& Spawnable
 		//SpawnInfo.Template = ObjectTemplate;
 		// allow pre-construction variables to be set.
 		SpawnInfo.bDeferConstruction = true;
+		SpawnInfo.Template = ObjectTemplate;
 		SpawnInfo.OverrideLevel = WorldContext->PersistentLevel;
 	}
 
@@ -63,6 +64,10 @@ UObject* FLevelSequenceActorSpawner::SpawnObject(FMovieSceneSpawnable& Spawnable
 		SpawnTransform.SetTranslation(RootComponent->RelativeLocation);
 		SpawnTransform.SetRotation(RootComponent->RelativeRotation.Quaternion());
 	}
+	else
+	{
+		SpawnTransform = Spawnable.SpawnTransform;
+	}
 
 	{
 		// Disable all particle components so that they don't auto fire as soon as the actor is spawned. The particles should be triggered through the particle track.
@@ -79,12 +84,12 @@ UObject* FLevelSequenceActorSpawner::SpawnObject(FMovieSceneSpawnable& Spawnable
 		return nullptr;
 	}
 	
-	UEngine::FCopyPropertiesForUnrelatedObjectsParams CopyParams;
-	CopyParams.bPreserveRootComponent = false;
-	CopyParams.bNotifyObjectReplacement = false;
-	SpawnedActor->UnregisterAllComponents();
-	UEngine::CopyPropertiesForUnrelatedObjects(ObjectTemplate, SpawnedActor, CopyParams);
-	SpawnedActor->RegisterAllComponents();
+	//UEngine::FCopyPropertiesForUnrelatedObjectsParams CopyParams;
+	//CopyParams.bPreserveRootComponent = false;
+	//CopyParams.bNotifyObjectReplacement = false;
+	//SpawnedActor->UnregisterAllComponents();
+	//UEngine::CopyPropertiesForUnrelatedObjects(ObjectTemplate, SpawnedActor, CopyParams);
+	//SpawnedActor->RegisterAllComponents();
 
 	// Ensure this spawnable is not a preview actor. Preview actors will not have BeginPlay() called on them.
 #if WITH_EDITOR
diff --git a/Engine/Source/Runtime/LevelSequence/Private/LevelSequencePlayer.cpp b/Engine/Source/Runtime/LevelSequence/Private/LevelSequencePlayer.cpp
--- a/Engine/Source/Runtime/LevelSequence/Private/LevelSequencePlayer.cpp
+++ b/Engine/Source/Runtime/LevelSequence/Private/LevelSequencePlayer.cpp
@@ -280,7 +280,7 @@ void ULevelSequencePlayer::TakeFrameSnapshot(FLevelSequencePlayerSnapshot& OutSn
 	OutSnapshot.Settings = SnapshotSettings;
 
 	OutSnapshot.MasterTime = CurrentTime;
-	OutSnapshot.MasterName = FText::FromString(Sequence->GetName());
+	OutSnapshot.MasterName = Sequence->GetName();
 
 	OutSnapshot.CurrentShotName = OutSnapshot.MasterName;
 	OutSnapshot.CurrentShotLocalTime = CurrentTime;
diff --git a/Engine/Source/Runtime/LevelSequence/Public/DefaultLevelSequenceInstanceData.h b/Engine/Source/Runtime/LevelSequence/Public/DefaultLevelSequenceInstanceData.h
new file mode 100644
--- /dev/null
+++ b/Engine/Source/Runtime/LevelSequence/Public/DefaultLevelSequenceInstanceData.h
@@ -0,0 +1,36 @@
+// Copyright 1998-2017 Epic Games, Inc. All Rights Reserved.
+
+#pragma once
+
+#include "ObjectMacros.h"
+#include "Transform.h"
+#include "IMovieSceneTransformOrigin.h"
+#include "GameFramework/Actor.h"
+
+#include "DefaultLevelSequenceInstanceData.generated.h"
+
+/** Default instance data class that level sequences understand. Implements IMovieSceneTransformOrigin. */
+UCLASS()
+class UDefaultLevelSequenceInstanceData
+	: public UObject
+	, public IMovieSceneTransformOrigin
+{
+	GENERATED_BODY()
+
+	UDefaultLevelSequenceInstanceData(const FObjectInitializer& Init)
+		: Super(Init)
+	{
+		TransformOriginActor = nullptr;
+		TransformOrigin = FTransform::Identity;
+	}
+
+	virtual FTransform NativeGetTransformOrigin() const override { return TransformOriginActor ? TransformOriginActor->ActorToWorld() : TransformOrigin; }
+
+	/** When set, this actor's world position will be used as the transform origin for all absolute transform sections */
+	UPROPERTY(EditAnywhere, Category="General")
+	AActor* TransformOriginActor;
+
+	/** Specifies a transform from which all absolute transform sections inside the sequence should be added to. Scale is ignored. */
+	UPROPERTY(EditAnywhere, Category="General")
+	FTransform TransformOrigin;
+};
\ No newline at end of file
diff --git a/Engine/Source/Runtime/LevelSequence/Public/LevelSequenceActor.h b/Engine/Source/Runtime/LevelSequence/Public/LevelSequenceActor.h
--- a/Engine/Source/Runtime/LevelSequence/Public/LevelSequenceActor.h
+++ b/Engine/Source/Runtime/LevelSequence/Public/LevelSequenceActor.h
@@ -53,7 +53,7 @@ protected:
 /**
  * Actor responsible for controlling a specific level sequence in the world.
  */
-UCLASS(hideCategories=(Rendering, Physics, LOD, Activation))
+UCLASS(hideCategories=(Rendering, Physics, LOD, Activation, Input))
 class LEVELSEQUENCE_API ALevelSequenceActor
 	: public AActor
 	, public IMovieSceneBindingOwnerInterface
@@ -79,16 +79,23 @@ public:
 	UPROPERTY(EditAnywhere, BlueprintReadOnly, Category="General", meta=(AllowedClasses="LevelSequence"))
 	FSoftObjectPath LevelSequence;
 
-	UPROPERTY(EditAnywhere, BlueprintReadOnly, AdvancedDisplay, Category="General")
+	UPROPERTY(EditAnywhere, BlueprintReadOnly, Category="General")
 	TArray<AActor*> AdditionalEventReceivers;
 
-	UPROPERTY(Instanced, VisibleAnywhere, AdvancedDisplay, BlueprintReadOnly, Category="General")
+	UPROPERTY(Instanced, BlueprintReadOnly, Category="General")
 	ULevelSequenceBurnInOptions* BurnInOptions;
 
 	/** Mapping of actors to override the sequence bindings with */
-	UPROPERTY(Instanced, VisibleAnywhere, AdvancedDisplay, BlueprintReadOnly, Category="General")
+	UPROPERTY(Instanced, BlueprintReadOnly, Category="General")
 	UMovieSceneBindingOverrides* BindingOverrides;
 
+	UPROPERTY(EditAnywhere, BlueprintReadOnly, Category="General")
+	bool bOverrideInstanceData;
+
+	/** Instance data that can be used to dynamically control sequence evaluation at runtime */
+	UPROPERTY(Instanced, BlueprintReadWrite, Category="General")
+	UObject* DefaultInstanceData;
+
 public:
 
 	/**
@@ -139,7 +146,7 @@ public:
 	UFUNCTION(BlueprintCallable, Category="Game|Cinematic|Bindings")
 	void AddBinding(FMovieSceneObjectBindingID Binding, AActor* Actor, bool bAllowBindingsFromAsset = false)
 	{
-		BindingOverrides->AddBinding(Binding, Actor);
+		BindingOverrides->AddBinding(Binding, Actor, bAllowBindingsFromAsset);
 		if (SequencePlayer)
 		{
 			SequencePlayer->State.Invalidate(Binding.GetGuid(), Binding.GetSequenceID());
diff --git a/Engine/Source/Runtime/LevelSequence/Public/LevelSequencePlayer.h b/Engine/Source/Runtime/LevelSequence/Public/LevelSequencePlayer.h
--- a/Engine/Source/Runtime/LevelSequence/Public/LevelSequencePlayer.h
+++ b/Engine/Source/Runtime/LevelSequence/Public/LevelSequencePlayer.h
@@ -58,13 +58,13 @@ struct FLevelSequencePlayerSnapshot
 	GENERATED_BODY()
 
 	UPROPERTY(VisibleAnywhere, BlueprintReadOnly, Category="General")
-	FText MasterName;
+	FString MasterName;
 
 	UPROPERTY(VisibleAnywhere, BlueprintReadOnly, Category="General")
 	float MasterTime;
 
 	UPROPERTY(VisibleAnywhere, BlueprintReadOnly, Category="General")
-	FText CurrentShotName;
+	FString CurrentShotName;
 
 	UPROPERTY(VisibleAnywhere, BlueprintReadOnly, Category="General")
 	float CurrentShotLocalTime;
diff --git a/Engine/Source/Runtime/LevelSequence/Public/LevelSequenceSpawnRegister.h b/Engine/Source/Runtime/LevelSequence/Public/LevelSequenceSpawnRegister.h
--- a/Engine/Source/Runtime/LevelSequence/Public/LevelSequenceSpawnRegister.h
+++ b/Engine/Source/Runtime/LevelSequence/Public/LevelSequenceSpawnRegister.h
@@ -4,7 +4,6 @@
 
 #include "CoreMinimal.h"
 #include "MovieSceneSpawnRegister.h"
-#include "MovieSceneSequenceInstance.h"
 
 class IMovieScenePlayer;
 class IMovieSceneObjectSpawner;
diff --git a/Engine/Source/Runtime/Media/Private/MediaModule.cpp b/Engine/Source/Runtime/Media/Private/MediaModule.cpp
--- a/Engine/Source/Runtime/Media/Private/MediaModule.cpp
+++ b/Engine/Source/Runtime/Media/Private/MediaModule.cpp
@@ -123,7 +123,7 @@ public:
 
 	virtual void ShutdownModule() override
 	{
-		if (TickerThread)
+		if (TickerThread != nullptr)
 		{
 			TickerThread->Kill(true);
 			delete TickerThread;
diff --git a/Engine/Source/Runtime/MediaAssets/Private/Assets/MediaPlayer.cpp b/Engine/Source/Runtime/MediaAssets/Private/Assets/MediaPlayer.cpp
--- a/Engine/Source/Runtime/MediaAssets/Private/Assets/MediaPlayer.cpp
+++ b/Engine/Source/Runtime/MediaAssets/Private/Assets/MediaPlayer.cpp
@@ -66,12 +66,16 @@ bool UMediaPlayer::CanPlaySource(UMediaSource* MediaSource)
 		return false;
 	}
 
+	UE_LOG(LogMediaAssets, Verbose, TEXT("%s.CanPlaySource"), *GetFName().ToString(), *MediaSource->GetFName().ToString());
+
 	return PlayerFacade->CanPlayUrl(MediaSource->GetUrl(), MediaSource);
 }
 
 
 bool UMediaPlayer::CanPlayUrl(const FString& Url)
 {
+	UE_LOG(LogMediaAssets, Verbose, TEXT("%s.CanPlayUrl"), *GetFName().ToString(), *Url);
+
 	if (Url.IsEmpty())
 	{
 		return false;
@@ -83,6 +87,8 @@ bool UMediaPlayer::CanPlayUrl(const FString& Url)
 
 void UMediaPlayer::Close()
 {
+	UE_LOG(LogMediaAssets, VeryVerbose, TEXT("%s.Close"), *GetFName().ToString());
+
 	PlayerFacade->Close();
 
 	Playlist = NewObject<UMediaPlaylist>(GetTransientPackage(), NAME_None, RF_Transactional | RF_Transient);
@@ -315,14 +321,25 @@ bool UMediaPlayer::IsPreparing() const
 
 bool UMediaPlayer::IsReady() const
 {
+	UE_LOG(LogMediaAssets, VeryVerbose, TEXT("%s.IsReady"), *GetFName().ToString());
 	return PlayerFacade->IsReady();
 }
 
 
 bool UMediaPlayer::Next()
 {
+	UE_LOG(LogMediaAssets, Verbose, TEXT("%s.Next"), *GetFName().ToString());
+
+	check(Playlist != nullptr);
 	int32 RemainingAttempts = Playlist->Num();
 
+	if (RemainingAttempts == 0)
+	{
+		return false;
+	}
+
+	PlayOnNext |= PlayerFacade->IsPlaying();
+
 	while (RemainingAttempts-- > 0)
 	{
 		UMediaSource* NextSource = Shuffle
@@ -343,6 +360,10 @@ bool UMediaPlayer::OpenFile(const FString& FilePath)
 {
 	Close();
 
+	UE_LOG(LogMediaAssets, Verbose, TEXT("%s.OpenFile %s"), *GetFName().ToString(), *FilePath);
+
+	check(Playlist != nullptr);
+
 	if (!Playlist->AddFile(FilePath))
 	{
 		return false;
@@ -358,10 +379,12 @@ bool UMediaPlayer::OpenPlaylistIndex(UMediaPlaylist* InPlaylist, int32 Index)
 
 	if (InPlaylist == nullptr)
 	{
-		UE_LOG(LogMediaAssets, Warning, TEXT("UMediaPlayer::OpenPlaylistIndex called with null MediaPlaylist"));
+		UE_LOG(LogMediaAssets, Warning, TEXT("%s.OpenPlaylistIndex called with null MediaPlaylist"), *GetFName().ToString());
 		return false;
 	}
 
+	UE_LOG(LogMediaAssets, Verbose, TEXT("%s.OpenSource %s %i"), *GetFName().ToString(), *InPlaylist->GetFName().ToString(), Index);
+
 	Playlist = InPlaylist;
 
 	if (Index == INDEX_NONE)
@@ -373,7 +396,7 @@ bool UMediaPlayer::OpenPlaylistIndex(UMediaPlaylist* InPlaylist, int32 Index)
 
 	if (MediaSource == nullptr)
 	{
-		UE_LOG(LogMediaAssets, Warning, TEXT("UMediaPlayer::OpenPlaylistIndex called with invalid PlaylistIndex %i"), Index);
+		UE_LOG(LogMediaAssets, Warning, TEXT("%s.OpenPlaylistIndex called with invalid PlaylistIndex %i"), *GetFName().ToString(), Index);
 		return false;
 	}
 
@@ -395,16 +418,19 @@ bool UMediaPlayer::OpenSource(UMediaSource* MediaSource)
 
 	if (MediaSource == nullptr)
 	{
-		UE_LOG(LogMediaAssets, Warning, TEXT("UMediaPlayer::OpenSource called with null MediaSource"));
+		UE_LOG(LogMediaAssets, Warning, TEXT("%s.OpenSource called with null MediaSource"), *GetFName().ToString());
 		return false;
 	}
 
+	UE_LOG(LogMediaAssets, Verbose, TEXT("%s.OpenSource %s"), *GetFName().ToString(), *MediaSource->GetFName().ToString());
+
 	if (!MediaSource->Validate())
 	{
 		UE_LOG(LogMediaAssets, Error, TEXT("Failed to validate media source %s (%s)"), *MediaSource->GetName(), *MediaSource->GetUrl());
 		return false;
 	}
 
+	check(Playlist != nullptr);
 	Playlist->Add(MediaSource);
 
 	return Next();
@@ -415,6 +441,10 @@ bool UMediaPlayer::OpenUrl(const FString& Url)
 {
 	Close();
 
+	UE_LOG(LogMediaAssets, Verbose, TEXT("%s.OpenUrl %s"), *GetFName().ToString(), *Url);
+
+	check(Playlist != nullptr);
+
 	if (!Playlist->AddUrl(Url))
 	{
 		return false;
@@ -426,20 +456,32 @@ bool UMediaPlayer::OpenUrl(const FString& Url)
 
 bool UMediaPlayer::Pause()
 {
+	UE_LOG(LogMediaAssets, Verbose, TEXT("%s.Pause"), *GetFName().ToString());
 	return PlayerFacade->SetRate(0.0f);
 }
 
 
 bool UMediaPlayer::Play()
 {
+	UE_LOG(LogMediaAssets, Verbose, TEXT("%s.Play"), *GetFName().ToString());
 	return PlayerFacade->SetRate(1.0f);
 }
 
 
 bool UMediaPlayer::Previous()
 {
+	UE_LOG(LogMediaAssets, Verbose, TEXT("%s.Previous"), *GetFName().ToString());
+
+	check(Playlist != nullptr);
 	int32 RemainingAttempts = Playlist->Num();
 
+	if (RemainingAttempts == 0)
+	{
+		return false;
+	}
+
+	PlayOnNext |= PlayerFacade->IsPlaying();
+
 	while (--RemainingAttempts >= 0)
 	{
 		UMediaSource* PrevSource = Shuffle
@@ -458,41 +500,43 @@ bool UMediaPlayer::Previous()
 
 bool UMediaPlayer::Reopen()
 {
-	if (Playlist == nullptr)
-	{
-		return false;
-	}
-
+	UE_LOG(LogMediaAssets, Verbose, TEXT("%s.Reopen"), *GetFName().ToString());
 	return OpenPlaylistIndex(Playlist, PlaylistIndex);
 }
 
 
 bool UMediaPlayer::Rewind()
 {
+	UE_LOG(LogMediaAssets, Verbose, TEXT("%s.Rewind"), *GetFName().ToString());
 	return Seek(FTimespan::Zero());
 }
 
 
 bool UMediaPlayer::Seek(const FTimespan& Time)
 {
+	UE_LOG(LogMediaAssets, VeryVerbose, TEXT("%s.Seek %s"), *GetFName().ToString(), *Time.ToString());
 	return PlayerFacade->Seek(Time);
 }
 
 
 bool UMediaPlayer::SelectTrack(EMediaPlayerTrack TrackType, int32 TrackIndex)
 {
+	UE_LOG(LogMediaAssets, Verbose, TEXT("%s.SelectTrack %s %i"), *GetFName().ToString(), *UEnum::GetValueAsString(TEXT("MediaAssets.EMediaPlayerTrack"), TrackType), TrackIndex);
 	return PlayerFacade->SelectTrack((EMediaTrackType)TrackType, TrackIndex);
 }
 
 
 void UMediaPlayer::SetDesiredPlayerName(FName PlayerName)
 {
+	UE_LOG(LogMediaAssets, Verbose, TEXT("%s.SetDesiredPlayerName %s"), *GetFName().ToString(), *PlayerName.ToString());
 	PlayerFacade->DesiredPlayerName = PlayerName;
 }
 
 
 bool UMediaPlayer::SetLooping(bool Looping)
 {
+	UE_LOG(LogMediaAssets, Verbose, TEXT("%s.SetLooping %s"), *GetFName().ToString(), *(Looping ? GTrue : GFalse).ToString());
+
 	Loop = Looping;
 
 	return PlayerFacade->SetLooping(Looping);
@@ -501,29 +545,34 @@ bool UMediaPlayer::SetLooping(bool Looping)
 
 bool UMediaPlayer::SetRate(float Rate)
 {
+	UE_LOG(LogMediaAssets, Verbose, TEXT("%s.SetRate %f"), *GetFName().ToString(), Rate);
 	return PlayerFacade->SetRate(Rate);
 }
 
 
 bool UMediaPlayer::SetTrackFormat(EMediaPlayerTrack TrackType, int32 TrackIndex, int32 FormatIndex)
 {
+	UE_LOG(LogMediaAssets, Verbose, TEXT("%s.SetTrackFormat %s %i %i"), *GetFName().ToString(), *UEnum::GetValueAsString(TEXT("MediaAssets.EMediaPlayerTrack"), TrackType), TrackIndex, FormatIndex);
 	return PlayerFacade->SetTrackFormat((EMediaTrackType)TrackType, TrackIndex, FormatIndex);
 }
 
 
 bool UMediaPlayer::SetVideoTrackFrameRate(int32 TrackIndex, int32 FormatIndex, float FrameRate)
 {
+	UE_LOG(LogMediaAssets, Verbose, TEXT("%s.SetVideoTrackFrameRate %i %i %f"), *GetFName().ToString(), TrackIndex, FormatIndex, FrameRate);
 	return PlayerFacade->SetVideoTrackFrameRate(TrackIndex, FormatIndex, FrameRate);
 }
 
 bool UMediaPlayer::SetViewField(float Horizontal, float Vertical, bool Absolute)
 {
+	UE_LOG(LogMediaAssets, Verbose, TEXT("%s.SetViewField %f %f %s"), *GetFName().ToString(), Horizontal, Vertical, *(Absolute ? GTrue : GFalse).ToString());
 	return PlayerFacade->SetViewField(Horizontal, Vertical, Absolute);
 }
 
 
 bool UMediaPlayer::SetViewRotation(const FRotator& Rotation, bool Absolute)
 {
+	UE_LOG(LogMediaAssets, Verbose, TEXT("%s.SetViewRotation %s %s"), *GetFName().ToString(), *Rotation.ToString(), *(Absolute ? GTrue : GFalse).ToString());
 	return PlayerFacade->SetViewOrientation(FQuat(Rotation), Absolute);
 }
 
@@ -550,6 +599,8 @@ bool UMediaPlayer::SupportsSeeking() const
 
 void UMediaPlayer::PausePIE()
 {
+	UE_LOG(LogMediaAssets, Verbose, TEXT("%s.PausePIE"), *GetFName().ToString());
+
 	WasPlayingInPIE = IsPlaying();
 
 	if (WasPlayingInPIE)
@@ -561,6 +612,8 @@ void UMediaPlayer::PausePIE()
 
 void UMediaPlayer::ResumePIE()
 {
+	UE_LOG(LogMediaAssets, Verbose, TEXT("%s.ResumePIE"), *GetFName().ToString());
+
 	if (WasPlayingInPIE)
 	{
 		Play();
@@ -628,6 +681,7 @@ void UMediaPlayer::PostLoad()
 	PlayerFacade->SetGuid(PlayerGuid);
 }
 
+
 #if WITH_EDITOR
 
 void UMediaPlayer::PostEditChangeProperty(FPropertyChangedEvent& PropertyChangedEvent)
@@ -687,6 +741,8 @@ void UMediaPlayer::HandlePlayerMediaEvent(EMediaEvent Event)
 	case EMediaEvent::PlaybackEndReached:
 		OnEndReached.Broadcast();
 
+		check(Playlist != nullptr);
+
 		if ((Loop && (Playlist->Num() != 1)) || (PlaylistIndex + 1 < Playlist->Num()))
 		{
 			PlayOnNext = true;
diff --git a/Engine/Source/Runtime/MediaAssets/Private/Assets/MediaSoundComponent.cpp b/Engine/Source/Runtime/MediaAssets/Private/Assets/MediaSoundComponent.cpp
--- a/Engine/Source/Runtime/MediaAssets/Private/Assets/MediaSoundComponent.cpp
+++ b/Engine/Source/Runtime/MediaAssets/Private/Assets/MediaSoundComponent.cpp
@@ -17,10 +17,16 @@
 UMediaSoundComponent::UMediaSoundComponent(const FObjectInitializer& ObjectInitializer)
 	: Super(ObjectInitializer)
 	, Channels(EMediaSoundChannels::Stereo)
+	, CachedRate(0.0f)
+	, CachedTime(FTimespan::Zero())
 	, Resampler(new FMediaAudioResampler)
 {
 	PrimaryComponentTick.bCanEverTick = true;
 	bAutoActivate = true;
+
+#if PLATFORM_MAC
+	PreferredBufferLength = 2048; // increase buffer callback size on macOS to prevent underruns
+#endif
 }
 
 
@@ -33,16 +39,25 @@ UMediaSoundComponent::~UMediaSoundComponent()
 /* UMediaSoundComponent interface
  *****************************************************************************/
 
+void UMediaSoundComponent::SetMediaPlayer(UMediaPlayer* NewMediaPlayer)
+{
+	CurrentPlayer = NewMediaPlayer;
+}
+
+
 void UMediaSoundComponent::UpdatePlayer()
 {
-	if (MediaPlayer == nullptr)
+	if (!CurrentPlayer.IsValid())
 	{
+		CachedRate = 0.0f;
+		CachedTime = FTimespan::Zero();
 		SampleQueue.Reset();
+
 		return;
 	}
 
 	// create a new sample queue if the player changed
-	TSharedRef<FMediaPlayerFacade, ESPMode::ThreadSafe> PlayerFacade = MediaPlayer->GetPlayerFacade();
+	TSharedRef<FMediaPlayerFacade, ESPMode::ThreadSafe> PlayerFacade = CurrentPlayer->GetPlayerFacade();
 
 	if (PlayerFacade != CurrentPlayerFacade)
 	{
@@ -57,6 +72,10 @@ void UMediaSoundComponent::UpdatePlayer()
 		CurrentPlayerFacade = PlayerFacade;
 	}
 
+	// caching play rate and time for audio thread (eventual consistency is sufficient)
+	CachedRate = PlayerFacade->GetRate();
+	CachedTime = PlayerFacade->GetTime();
+
 	check(SampleQueue.IsValid());
 }
 
@@ -72,7 +91,7 @@ void UMediaSoundComponent::TickComponent(float DeltaTime, enum ELevelTick TickTy
 }
 
 
-/* USynthComponent interface
+/* USceneComponent interface
  *****************************************************************************/
 
 void UMediaSoundComponent::Activate(bool bReset)
@@ -97,6 +116,41 @@ void UMediaSoundComponent::Deactivate()
 }
 
 
+/* UObject interface
+ *****************************************************************************/
+
+void UMediaSoundComponent::PostLoad()
+{
+	Super::PostLoad();
+
+	CurrentPlayer = MediaPlayer;
+}
+
+
+#if WITH_EDITOR
+
+void UMediaSoundComponent::PostEditChangeProperty(FPropertyChangedEvent& PropertyChangedEvent)
+{
+	static const FName MediaPlayerName = GET_MEMBER_NAME_CHECKED(UMediaSoundComponent, MediaPlayer);
+
+	UProperty* PropertyThatChanged = PropertyChangedEvent.Property;
+
+	if (PropertyThatChanged != nullptr)
+	{
+		const FName PropertyName = PropertyThatChanged->GetFName();
+
+		if (PropertyName == MediaPlayerName)
+		{
+			CurrentPlayer = MediaPlayer;
+		}
+	}
+
+	Super::PostEditChangeProperty(PropertyChangedEvent);
+}
+
+#endif //WITH_EDITOR
+
+
 /* USynthComponent interface
  *****************************************************************************/
 
@@ -124,19 +178,16 @@ bool UMediaSoundComponent::Init(int32& SampleRate)
 
 void UMediaSoundComponent::OnGenerateAudio(float* OutAudio, int32 NumSamples)
 {
-	TSharedPtr<FMediaPlayerFacade, ESPMode::ThreadSafe> PinnedPlayerFacade;
 	TSharedPtr<FMediaAudioSampleQueue, ESPMode::ThreadSafe> PinnedSampleQueue;
 	{
 		FScopeLock Lock(&CriticalSection);
-
-		PinnedPlayerFacade = CurrentPlayerFacade.Pin();
 		PinnedSampleQueue = SampleQueue;
 	}
 
-	if (PinnedPlayerFacade.IsValid() && PinnedPlayerFacade->IsPlaying() && PinnedSampleQueue.IsValid())
+	if (PinnedSampleQueue.IsValid() && (CachedRate != 0.0f))
 	{
 		const uint32 FramesRequested = NumSamples / NumChannels;
-		const uint32 FramesWritten = Resampler->Generate(OutAudio, FramesRequested, PinnedPlayerFacade->GetRate(), PinnedPlayerFacade->GetTime(), *PinnedSampleQueue);
+		const uint32 FramesWritten = Resampler->Generate(OutAudio, FramesRequested, CachedRate, CachedTime, *PinnedSampleQueue);
 	}
 	else
 	{
diff --git a/Engine/Source/Runtime/MediaAssets/Private/Assets/MediaTexture.cpp b/Engine/Source/Runtime/MediaAssets/Private/Assets/MediaTexture.cpp
--- a/Engine/Source/Runtime/MediaAssets/Private/Assets/MediaTexture.cpp
+++ b/Engine/Source/Runtime/MediaAssets/Private/Assets/MediaTexture.cpp
@@ -57,9 +57,9 @@ UMediaTexture::UMediaTexture(const FObjectInitializer& ObjectInitializer)
 	: Super(ObjectInitializer)
 	, AddressX(TA_Clamp)
 	, AddressY(TA_Clamp)
-	, AutoClear(true)
+	, AutoClear(false)
 	, ClearColor(FLinearColor::Black)
-	, MediaPlayer(nullptr)
+	, DefaultGuid(FGuid::NewGuid())
 	, Dimensions(FIntPoint::ZeroValue)
 	, Size(0)
 {
@@ -87,12 +87,24 @@ int32 UMediaTexture::GetHeight() const
 }
 
 
+UMediaPlayer* UMediaTexture::GetMediaPlayer() const
+{
+	return CurrentPlayer.Get();
+}
+
+
 int32 UMediaTexture::GetWidth() const
 {
 	return Dimensions.X;
 }
 
 
+void UMediaTexture::SetMediaPlayer(UMediaPlayer* NewMediaPlayer)
+{
+	CurrentPlayer = NewMediaPlayer;
+}
+
+
 /* UTexture interface
  *****************************************************************************/
 
@@ -109,7 +121,7 @@ FTextureResource* UMediaTexture::CreateResource()
 		}
 	}
 
-	return new FMediaTextureResource(*this, Dimensions, Size);
+	return new FMediaTextureResource(*this, Dimensions, Size, ClearColor, CurrentGuid);
 }
 
 
@@ -133,7 +145,7 @@ float UMediaTexture::GetSurfaceHeight() const
 
 FGuid UMediaTexture::GetExternalTextureGuid() const
 {
-	return MediaPlayer ? MediaPlayer->GetGuid() : FGuid();
+	return CurrentGuid;
 }
 
 
@@ -154,7 +166,14 @@ void UMediaTexture::BeginDestroy()
 		ClockSink.Reset();
 	}
 
-	UnregisterPlayerGuid();
+	if (CurrentGuid.IsValid())
+	{
+		ENQUEUE_UNIQUE_RENDER_COMMAND_ONEPARAMETER(MediaTextureUnregisterGuid,
+			FGuid, Guid, CurrentGuid,
+			{
+				FExternalTextureRegistry::Get().UnregisterExternalTexture(Guid);
+			});
+	}
 
 	Super::BeginDestroy();
 }
@@ -174,6 +193,14 @@ void UMediaTexture::GetResourceSizeEx(FResourceSizeEx& CumulativeResourceSize)
 }
 
 
+void UMediaTexture::PostLoad()
+{
+	Super::PostLoad();
+
+	CurrentPlayer = MediaPlayer;
+}
+
+
 #if WITH_EDITOR
 
 void UMediaTexture::PostEditChangeProperty(FPropertyChangedEvent& PropertyChangedEvent)
@@ -195,6 +222,11 @@ void UMediaTexture::PostEditChangeProperty(FPropertyChangedEvent& PropertyChange
 
 	const FName PropertyName = PropertyThatChanged->GetFName();
 
+	if (PropertyName == MediaPlayerName)
+	{
+		CurrentPlayer = MediaPlayer;
+	}
+
 	// don't update resource for these properties
 	if ((PropertyName == AutoClearName) ||
 		(PropertyName == ClearColorName) ||
@@ -223,58 +255,68 @@ void UMediaTexture::PostEditChangeProperty(FPropertyChangedEvent& PropertyChange
 
 void UMediaTexture::TickResource(FTimespan Timecode)
 {
-	if ((MediaPlayer == nullptr) || (Resource == nullptr))
+	if (Resource == nullptr)
 	{
-		CurrentPlayerFacade.Reset();
-		SampleQueue.Reset();
-
 		return;
 	}
 
-	// create a new sample queue if the player changed
-	TSharedRef<FMediaPlayerFacade, ESPMode::ThreadSafe> PlayerFacade = MediaPlayer->GetPlayerFacade();
+	const FGuid LastGuid = CurrentGuid;
 
-	if (PlayerFacade != CurrentPlayerFacade)
+	// media player bookkeeping
+	if (CurrentPlayer.IsValid())
 	{
-		SampleQueue = MakeShared<FMediaTextureSampleQueue, ESPMode::ThreadSafe>();
-		PlayerFacade->AddVideoSampleSink(SampleQueue.ToSharedRef());
-		CurrentPlayerFacade = PlayerFacade;
-	}
+		const FGuid PlayerGuid = CurrentPlayer->GetGuid();
 
-	check(SampleQueue.IsValid());
-
-	// unregister previous external texture GUID if needed
-	const FGuid PlayerGuid = MediaPlayer->GetGuid();
-
-	if (PlayerGuid != LastPlayerGuid)
+		if (CurrentGuid != PlayerGuid)
+		{
+			SampleQueue = MakeShared<FMediaTextureSampleQueue, ESPMode::ThreadSafe>();
+			CurrentPlayer->GetPlayerFacade()->AddVideoSampleSink(SampleQueue.ToSharedRef());
+			CurrentGuid = PlayerGuid;
+		}
+	}
+	else if (CurrentGuid != DefaultGuid)
 	{
-		UnregisterPlayerGuid();
-		LastPlayerGuid = PlayerGuid;
+		SampleQueue.Reset();
+		CurrentGuid = DefaultGuid;
 	}
-
-	// retain the last rendered frame if player is inactive
-	const bool PlayerActive = MediaPlayer->IsPaused() || MediaPlayer->IsPlaying() || MediaPlayer->IsPreparing();
-
-	if (!PlayerActive && !AutoClear)
+	else if ((LastClearColor == ClearColor) && (LastSrgb == SRGB))
 	{
-		return;
+		return; // nothing to render
 	}
 
-	// issue a render command to render the current sample
+	LastClearColor = ClearColor;
+	LastSrgb = SRGB;
+
+	// set up render parameters
 	FMediaTextureResource::FRenderParams RenderParams;
+
+	if (CurrentPlayer.IsValid())
 	{
-		RenderParams.ClearColor = ClearColor;
-		RenderParams.PlayerGuid = PlayerGuid;
-		RenderParams.Rate = MediaPlayer->GetRate();
-		RenderParams.SrgbOutput = SRGB;
-		RenderParams.Time = MediaPlayer->GetTime();
+		const bool PlayerActive = CurrentPlayer->IsPaused() || CurrentPlayer->IsPlaying() || CurrentPlayer->IsPreparing();
 
 		if (PlayerActive)
 		{
 			RenderParams.SampleSource = SampleQueue;
 		}
+		else if (!AutoClear)
+		{
+			return; // retain last frame
+		}
+
+		RenderParams.Rate = CurrentPlayer->GetRate();
+		RenderParams.Time = CurrentPlayer->GetTime();
+	}
+	else if (!AutoClear && LastGuid.IsValid())
+	{
+		return; // retain last frame
 	}
 
+	RenderParams.ClearColor = ClearColor;
+	RenderParams.LastGuid = LastGuid;
+	RenderParams.TextureGuid = CurrentGuid;
+	RenderParams.SrgbOutput = SRGB;
+
+	// redraw texture resource on render thread
 	ENQUEUE_UNIQUE_RENDER_COMMAND_TWOPARAMETER(MediaTextureResourceRender,
 		FMediaTextureResource*, ResourceParam, (FMediaTextureResource*)Resource,
 		FMediaTextureResource::FRenderParams, RenderParamsParam, RenderParams,
@@ -282,18 +324,3 @@ void UMediaTexture::TickResource(FTimespan Timecode)
 			ResourceParam->Render(RenderParamsParam);
 		});
 }
-
-
-void UMediaTexture::UnregisterPlayerGuid()
-{
-	if (!LastPlayerGuid.IsValid())
-	{
-		return;
-	}
-
-	ENQUEUE_UNIQUE_RENDER_COMMAND_ONEPARAMETER(MediaTextureUnregisterPlayerGuid,
-		FGuid, PlayerGuid, LastPlayerGuid,
-		{
-			FExternalTextureRegistry::Get().UnregisterExternalTexture(PlayerGuid);
-		});
-}
