[MEDIA] [~] Increase media sound component samples generated per callback and re-introduce pre-buffering in sound wave procedural to reduce the impact of hitching on the game thread for media.

diff --git a/Engine/Source/Runtime/Engine/Private/SoundWaveProcedural.cpp b/Engine/Source/Runtime/Engine/Private/SoundWaveProcedural.cpp
--- a/Engine/Source/Runtime/Engine/Private/SoundWaveProcedural.cpp
+++ b/Engine/Source/Runtime/Engine/Private/SoundWaveProcedural.cpp
@@ -128,7 +128,7 @@ int32 USoundWaveProcedural::GeneratePCMData(uint8* PCMData, const int32 SamplesN
 	SamplesAvailable = AudioBuffer.Num() / SampleByteSize;
 
 	// Wait until we have enough samples that are requested before starting.
-	if (SamplesAvailable >= 0)
+	if (SamplesAvailable >= SamplesToGenerate)
 	{
 		const int32 SamplesToCopy = FMath::Min<int32>(SamplesToGenerate, SamplesAvailable);
 		const int32 BytesToCopy = SamplesToCopy * SampleByteSize;
diff --git a/Engine/Source/Runtime/MediaAssets/Private/Assets/MediaSoundComponent.cpp b/Engine/Source/Runtime/MediaAssets/Private/Assets/MediaSoundComponent.cpp
--- a/Engine/Source/Runtime/MediaAssets/Private/Assets/MediaSoundComponent.cpp
+++ b/Engine/Source/Runtime/MediaAssets/Private/Assets/MediaSoundComponent.cpp
@@ -47,10 +47,6 @@ UMediaSoundComponent::UMediaSoundComponent(const FObjectInitializer& ObjectIniti
 	PrimaryComponentTick.bCanEverTick = true;
 	bAutoActivate = true;
 
-#if PLATFORM_MAC
-	PreferredBufferLength = 2048; // increase buffer callback size on macOS to prevent underruns
-#endif
-
 #if WITH_EDITORONLY_DATA
 	bVisualizeComponent = true;
 #endif
@@ -278,6 +274,9 @@ bool UMediaSoundComponent::Init(int32& SampleRate)
 		NumChannels = 8;
 	}*/
 
+	// increase buffer callback size for media decoding. Media doesn't need fast response time so can decode more per callback.
+	PreferredBufferLength = NumChannels * 8196;
+
 	Resampler->Initialize(NumChannels, SampleRate);
 
 	return true;
