[release-4.20] #jira UE-60459 bring over Rex's fix for media sound component stuttering to 4.20. #fyi Rex.Hill Aaron.McLeran #rb none

diff --git a/Engine/Source/Runtime/MediaAssets/Private/Assets/MediaSoundComponent.cpp b/Engine/Source/Runtime/MediaAssets/Private/Assets/MediaSoundComponent.cpp
--- a/Engine/Source/Runtime/MediaAssets/Private/Assets/MediaSoundComponent.cpp
+++ b/Engine/Source/Runtime/MediaAssets/Private/Assets/MediaSoundComponent.cpp
@@ -36,7 +36,7 @@ USoundClass* UMediaSoundComponent::DefaultMediaSoundClassObject = nullptr;
 UMediaSoundComponent::UMediaSoundComponent(const FObjectInitializer& ObjectInitializer)
 	: Super(ObjectInitializer)
 	, Channels(EMediaSoundChannels::Stereo)
-	, DynamicRateAdjustment(true)
+	, DynamicRateAdjustment(false)
 	, RateAdjustmentFactor(0.00000001f)
 	, RateAdjustmentRange(FFloatRange(0.995f, 1.005f))
 	, CachedRate(0.0f)
