Merging //UE4/Dev-Main to Dev-Build (//UE4/Dev-Build)

#rnx

diff --git a/Engine/Plugins/Online/OnlineSubsystemUtils/Source/OnlineSubsystemUtils/Private/VoiceEngineImpl.cpp b/Engine/Plugins/Online/OnlineSubsystemUtils/Source/OnlineSubsystemUtils/Private/VoiceEngineImpl.cpp
--- a/Engine/Plugins/Online/OnlineSubsystemUtils/Source/OnlineSubsystemUtils/Private/VoiceEngineImpl.cpp
+++ b/Engine/Plugins/Online/OnlineSubsystemUtils/Source/OnlineSubsystemUtils/Private/VoiceEngineImpl.cpp
@@ -507,10 +507,7 @@ uint32 FVoiceEngineImpl::SubmitRemoteVoiceData(const FUniqueNetIdWrapper& Remote
 	// Generate a streaming wave audio component for voice playback
 	if (QueuedData.VoipSynthComponent == nullptr || QueuedData.VoipSynthComponent->IsPendingKill())
 	{
-		if (SerializeHelper == nullptr)
-		{
-			SerializeHelper = new FVoiceSerializeHelper(this);
-		}
+		CreateSerializeHelper();
 
 		QueuedData.VoipSynthComponent = CreateVoiceSynthComponent(UVOIPStatics::GetVoiceSampleRate());
 		if (QueuedData.VoipSynthComponent)
@@ -762,4 +759,10 @@ int32 FVoiceEngineImpl::GetMaxVoiceRemainderSize()
 	return MAX_VOICE_REMAINDER_SIZE;
 }
 
-
+void FVoiceEngineImpl::CreateSerializeHelper()
+{
+	if (SerializeHelper == nullptr)
+	{
+		SerializeHelper = new FVoiceSerializeHelper(this);
+	}
+}
\ No newline at end of file
diff --git a/Engine/Plugins/Online/OnlineSubsystemUtils/Source/OnlineSubsystemUtils/Public/VoiceEngineImpl.h b/Engine/Plugins/Online/OnlineSubsystemUtils/Source/OnlineSubsystemUtils/Public/VoiceEngineImpl.h
--- a/Engine/Plugins/Online/OnlineSubsystemUtils/Source/OnlineSubsystemUtils/Public/VoiceEngineImpl.h
+++ b/Engine/Plugins/Online/OnlineSubsystemUtils/Source/OnlineSubsystemUtils/Public/VoiceEngineImpl.h
@@ -298,6 +298,7 @@ protected:
 	TArray<uint8>&             GetDecompressedVoiceBuffer() { return DecompressedVoiceBuffer; }
 	FLocalVoiceData*           GetLocalPlayerVoiceData()    { return PlayerVoiceData; }
 	int32                      GetMaxVoiceRemainderSize();
+	void					   CreateSerializeHelper();
 };
 
 typedef TSharedPtr<FVoiceEngineImpl, ESPMode::ThreadSafe> FVoiceEngineImplPtr;
diff --git a/Engine/Source/Editor/UnrealEd/Private/Commandlets/FixConflictingLocalizationKeys.cpp b/Engine/Source/Editor/UnrealEd/Private/Commandlets/FixConflictingLocalizationKeys.cpp
--- a/Engine/Source/Editor/UnrealEd/Private/Commandlets/FixConflictingLocalizationKeys.cpp
+++ b/Engine/Source/Editor/UnrealEd/Private/Commandlets/FixConflictingLocalizationKeys.cpp
@@ -178,7 +178,7 @@ bool ReKeyTextProperty(UStruct* InOuterType, void* InAddrToUpdate, const TArray<
 						// ContainerIndex is the element index, but we need the sparse index
 						int32 SparseIndex = 0;
 						{
-							const int32 ElementCount = ScriptMapHelper.Num();
+							const int32 ElementCount = ScriptMapHelper.GetMaxIndex();
 							for (int32 ElementIndex = 0; ElementIndex < ElementCount; ++SparseIndex)
 							{
 								if (ScriptMapHelper.IsValidIndex(SparseIndex))
@@ -215,7 +215,7 @@ bool ReKeyTextProperty(UStruct* InOuterType, void* InAddrToUpdate, const TArray<
 						// ContainerIndex is the element index, but we need the sparse index
 						int32 SparseIndex = 0;
 						{
-							const int32 ElementCount = ScriptSetHelper.Num();
+							const int32 ElementCount = ScriptSetHelper.GetMaxIndex();
 							for (int32 ElementIndex = 0; ElementIndex < ElementCount; ++SparseIndex)
 							{
 								if (ScriptSetHelper.IsValidIndex(SparseIndex))
diff --git a/Engine/Source/Programs/IOS/iPhonePackager/app.config b/Engine/Source/Programs/IOS/iPhonePackager/app.config
--- a/Engine/Source/Programs/IOS/iPhonePackager/app.config
+++ b/Engine/Source/Programs/IOS/iPhonePackager/app.config
@@ -1,7 +1,7 @@
 <?xml version="1.0"?>
 <configuration>
 	<startup>
-		<supportedRuntime version="v4.0" sku=".NETFramework,Version=v4.0"/>
+		<supportedRuntime version="v4.0" sku=".NETFramework,Version=v4.6.2"/>
 	</startup>
 	<runtime>
 		<generatePublisherEvidence enabled="false"/>
diff --git a/Engine/Source/Programs/IOS/iPhonePackager/iPhonePackager.csproj b/Engine/Source/Programs/IOS/iPhonePackager/iPhonePackager.csproj
--- a/Engine/Source/Programs/IOS/iPhonePackager/iPhonePackager.csproj
+++ b/Engine/Source/Programs/IOS/iPhonePackager/iPhonePackager.csproj
@@ -1,5 +1,5 @@
 ﻿<?xml version="1.0" encoding="utf-8"?>
-<Project ToolsVersion="4.0" DefaultTargets="Build" xmlns="http://schemas.microsoft.com/developer/msbuild/2003">
+<Project ToolsVersion="12.0" DefaultTargets="Build" xmlns="http://schemas.microsoft.com/developer/msbuild/2003">
   <PropertyGroup>
     <Configuration Condition=" '$(Configuration)' == '' ">Development</Configuration>
     <Platform Condition=" '$(Platform)' == '' ">AnyCPU</Platform>
@@ -10,7 +10,7 @@
     <AppDesignerFolder>Properties</AppDesignerFolder>
     <RootNamespace>iPhonePackager</RootNamespace>
     <AssemblyName>IPhonePackager</AssemblyName>
-    <TargetFrameworkVersion>v4.5</TargetFrameworkVersion>
+    <TargetFrameworkVersion>v4.6.2</TargetFrameworkVersion>
     <FileAlignment>512</FileAlignment>
     <SccProjectName>
     </SccProjectName>
diff --git a/Engine/Source/Programs/UnrealBuildTool/Modes/WriteMetadataMode.cs b/Engine/Source/Programs/UnrealBuildTool/Modes/WriteMetadataMode.cs
--- a/Engine/Source/Programs/UnrealBuildTool/Modes/WriteMetadataMode.cs
+++ b/Engine/Source/Programs/UnrealBuildTool/Modes/WriteMetadataMode.cs
@@ -122,15 +122,8 @@ namespace UnrealBuildTool
 						// Merge files from the existing manifests with the new ones
 						foreach(KeyValuePair<FileReference, ModuleManifest> Pair in PreviousFileToManifest)
 						{
-							ModuleManifest OldManifest = Pair.Value;
-							ModuleManifest NewManifest = TargetInfo.FileToManifest[Pair.Key];
-							foreach(KeyValuePair<string, string> ModulePair in OldManifest.ModuleNameToFileName)
-							{
-								if(!NewManifest.ModuleNameToFileName.ContainsKey(ModulePair.Key))
-								{
-									NewManifest.ModuleNameToFileName.Add(ModulePair.Key, ModulePair.Value);
-								}
-							}
+							ModuleManifest TargetManifest = TargetInfo.FileToManifest[Pair.Key];
+							MergeManifests(Pair.Value, TargetManifest);
 						}
 
 						// Update the build id to use the current one
@@ -144,6 +137,18 @@ namespace UnrealBuildTool
 					Receipt.Version.BuildId = Guid.NewGuid().ToString();
 				}
 			}
+			else
+			{
+				// Read all the manifests and merge them into the new ones, if they have the same build id
+				foreach(KeyValuePair<FileReference, ModuleManifest> Pair in TargetInfo.FileToManifest)
+				{
+					ModuleManifest SourceManifest;
+					if(TryReadManifest(Pair.Key, out SourceManifest) && SourceManifest.BuildId == Receipt.Version.BuildId)
+					{
+						MergeManifests(SourceManifest, Pair.Value);
+					}
+				}
+			}
 
 			// Update the build id in all the manifests, and write them out
 			foreach (KeyValuePair<FileReference, ModuleManifest> Pair in TargetInfo.FileToManifest)
@@ -224,31 +229,65 @@ namespace UnrealBuildTool
 			bool bCanRecycleManifests = true;
 			foreach(FileReference ManifestFileName in ManifestFiles)
 			{
-				if(ManifestFileName.IsUnderDirectory(UnrealBuildTool.EngineDirectory) && FileReference.Exists(ManifestFileName))
+				ModuleManifest Manifest;
+				if(ManifestFileName.IsUnderDirectory(UnrealBuildTool.EngineDirectory) && TryReadManifest(ManifestFileName, out Manifest))
 				{
-					try
+					if(Manifest.BuildId == BuildId)
 					{
-						ModuleManifest Manifest = ModuleManifest.Read(ManifestFileName);
-						if(Manifest.BuildId == BuildId)
+						if(IsOutOfDate(ManifestFileName, Manifest))
 						{
-							if(IsOutOfDate(ManifestFileName, Manifest))
-							{
-								bCanRecycleManifests = false;
-								break;
-							}
-							RecycleFileToManifest.Add(ManifestFileName, Manifest);
+							bCanRecycleManifests = false;
+							break;
 						}
-					}
-					catch(Exception Ex)
-					{
-						Log.TraceWarning("Unable to read '{0}'; ignoring.", ManifestFileName);
-						Log.TraceLog(ExceptionUtils.FormatExceptionDetails(Ex));
+						RecycleFileToManifest.Add(ManifestFileName, Manifest);
 					}
 				}
 			}
 			return bCanRecycleManifests;
 		}
 
+		/// <summary>
+		/// Attempts to read a manifest from the given location
+		/// </summary>
+		/// <param name="ManifestFileName">Path to the manifest</param>
+		/// <param name="Manifest">If successful, receives the manifest that was read</param>
+		/// <returns>True if the manifest was read correctly, false otherwise</returns>
+		public static bool TryReadManifest(FileReference ManifestFileName, out ModuleManifest Manifest)
+		{
+			if(FileReference.Exists(ManifestFileName))
+			{
+				try
+				{
+					Manifest = ModuleManifest.Read(ManifestFileName);
+					return true;
+				}
+				catch(Exception Ex)
+				{
+					Log.TraceWarning("Unable to read '{0}'; ignoring.", ManifestFileName);
+					Log.TraceLog(ExceptionUtils.FormatExceptionDetails(Ex));
+				}
+			}
+
+			Manifest = null;
+			return false;
+		}
+
+		/// <summary>
+		/// Merge a manifest into another manifest
+		/// </summary>
+		/// <param name="SourceManifest">The source manifest</param>
+		/// <param name="TargetManifest">The target manifest to merge into</param>
+		static void MergeManifests(ModuleManifest SourceManifest, ModuleManifest TargetManifest)
+		{
+			foreach(KeyValuePair<string, string> ModulePair in SourceManifest.ModuleNameToFileName)
+			{
+				if(!TargetManifest.ModuleNameToFileName.ContainsKey(ModulePair.Key))
+				{
+					TargetManifest.ModuleNameToFileName.Add(ModulePair.Key, ModulePair.Value);
+				}
+			}
+		}
+
 		/// <summary>
 		/// Checks whether a module manifest on disk is out of date (whether any of the binaries it references are newer than it is)
 		/// </summary>
diff --git a/Engine/Source/Runtime/CoreUObject/Private/Serialization/PropertyLocalizationDataGathering.cpp b/Engine/Source/Runtime/CoreUObject/Private/Serialization/PropertyLocalizationDataGathering.cpp
--- a/Engine/Source/Runtime/CoreUObject/Private/Serialization/PropertyLocalizationDataGathering.cpp
+++ b/Engine/Source/Runtime/CoreUObject/Private/Serialization/PropertyLocalizationDataGathering.cpp
@@ -268,21 +268,34 @@ void FPropertyLocalizationDataGatherer::GatherLocalizationDataFromChildTextPrope
 			const int32 ElementCount = ScriptArrayHelper.Num();
 			for(int32 j = 0; j < ElementCount; ++j)
 			{
-				GatherLocalizationDataFromChildTextProperties(PathToElement + FString::Printf(TEXT("(%d)"), j), ArrayProperty->Inner, ScriptArrayHelper.GetRawPtr(j), nullptr, ChildPropertyGatherTextFlags);
+				const uint8* ElementPtr = ScriptArrayHelper.GetRawPtr(j);
+				const uint8* DefaultElementPtr = nullptr;
+				if (DefaultElementValueAddress)
+				{
+					FScriptArrayHelper DefaultScriptArrayHelper(ArrayProperty, DefaultElementValueAddress);
+					if (DefaultScriptArrayHelper.IsValidIndex(j))
+					{
+						DefaultElementPtr = DefaultScriptArrayHelper.GetRawPtr(j);
+					}
+				}
+
+				GatherLocalizationDataFromChildTextProperties(PathToElement + FString::Printf(TEXT("(%d)"), j), ArrayProperty->Inner, ElementPtr, DefaultElementPtr, ChildPropertyGatherTextFlags);
 			}
 		}
 		// Property is a map property.
 		else if (MapProperty)
 		{
 			// Skip maps of types which can't contain any text.
-			if (CanGatherFromInnerProperty(MapProperty->KeyProp) == false && CanGatherFromInnerProperty(MapProperty->ValueProp) == false)
+			const bool bGatherMapKey = CanGatherFromInnerProperty(MapProperty->KeyProp);
+			const bool bGatherMapValue = CanGatherFromInnerProperty(MapProperty->ValueProp);
+			if (!bGatherMapKey && !bGatherMapValue)
 			{
 				continue;
 			}
 
 			// Iterate over all elements of the map.
 			FScriptMapHelper ScriptMapHelper(MapProperty, ElementValueAddress);
-			const int32 ElementCount = ScriptMapHelper.Num();
+			const int32 ElementCount = ScriptMapHelper.GetMaxIndex();
 			for(int32 j = 0, ElementIndex = 0; ElementIndex < ElementCount; ++j)
 			{
 				if (!ScriptMapHelper.IsValidIndex(j))
@@ -291,13 +304,25 @@ void FPropertyLocalizationDataGatherer::GatherLocalizationDataFromChildTextPrope
 				}
 
 				const uint8* MapPairPtr = ScriptMapHelper.GetPairPtr(j);
-				if (CanGatherFromInnerProperty(MapProperty->KeyProp))
+				const uint8* DefaultPairPtr = nullptr;
+				if (DefaultElementValueAddress)
 				{
-					GatherLocalizationDataFromChildTextProperties(PathToElement + FString::Printf(TEXT("(%d - Key)"), ElementIndex), MapProperty->KeyProp, MapPairPtr + MapProperty->MapLayout.KeyOffset, nullptr, ChildPropertyGatherTextFlags);
+					const uint8* MapKeyPtr = MapPairPtr + MapProperty->MapLayout.KeyOffset;
+					FScriptMapHelper DefaultScriptMapHelper(MapProperty, DefaultElementValueAddress);
+					DefaultPairPtr = DefaultScriptMapHelper.FindMapPairPtrFromHash(MapKeyPtr);
 				}
-				if (CanGatherFromInnerProperty(MapProperty->ValueProp))
+
+				if (bGatherMapKey)
 				{
-					GatherLocalizationDataFromChildTextProperties(PathToElement + FString::Printf(TEXT("(%d - Value)"), ElementIndex), MapProperty->ValueProp, MapPairPtr + MapProperty->MapLayout.ValueOffset, nullptr, ChildPropertyGatherTextFlags);
+					const uint8* MapKeyPtr = MapPairPtr + MapProperty->MapLayout.KeyOffset;
+					const uint8* DefaultMapKeyPtr = DefaultPairPtr ? DefaultPairPtr + MapProperty->MapLayout.KeyOffset : nullptr;
+					GatherLocalizationDataFromChildTextProperties(PathToElement + FString::Printf(TEXT("(%d - Key)"), ElementIndex), MapProperty->KeyProp, MapKeyPtr, DefaultMapKeyPtr, ChildPropertyGatherTextFlags);
+				}
+				if (bGatherMapValue)
+				{
+					const uint8* MapValuePtr = MapPairPtr + MapProperty->MapLayout.ValueOffset;
+					const uint8* DefaultMapValuePtr = DefaultPairPtr ? DefaultPairPtr + MapProperty->MapLayout.ValueOffset : nullptr;
+					GatherLocalizationDataFromChildTextProperties(PathToElement + FString::Printf(TEXT("(%d - Value)"), ElementIndex), MapProperty->ValueProp, MapValuePtr, DefaultMapValuePtr, ChildPropertyGatherTextFlags);
 				}
 				++ElementIndex;
 			}
@@ -313,7 +338,7 @@ void FPropertyLocalizationDataGatherer::GatherLocalizationDataFromChildTextPrope
 
 			// Iterate over all elements of the Set.
 			FScriptSetHelper ScriptSetHelper(SetProperty, ElementValueAddress);
-			const int32 ElementCount = ScriptSetHelper.Num();
+			const int32 ElementCount = ScriptSetHelper.GetMaxIndex();
 			for(int32 j = 0, ElementIndex = 0; ElementIndex < ElementCount; ++j)
 			{
 				if (!ScriptSetHelper.IsValidIndex(j))
@@ -322,7 +347,14 @@ void FPropertyLocalizationDataGatherer::GatherLocalizationDataFromChildTextPrope
 				}
 
 				const uint8* ElementPtr = ScriptSetHelper.GetElementPtr(j);
-				GatherLocalizationDataFromChildTextProperties(PathToElement + FString::Printf(TEXT("(%d)"), ElementIndex), SetProperty->ElementProp, ElementPtr + SetProperty->SetLayout.ElementOffset, nullptr, ChildPropertyGatherTextFlags);
+				const uint8* DefaultElementPtr = nullptr;
+				if (DefaultElementValueAddress)
+				{
+					FScriptSetHelper DefaultScriptSetHelper(SetProperty, DefaultElementValueAddress);
+					DefaultElementPtr = DefaultScriptSetHelper.FindElementPtrFromHash(ElementPtr);
+				}
+
+				GatherLocalizationDataFromChildTextProperties(PathToElement + FString::Printf(TEXT("(%d)"), ElementIndex), SetProperty->ElementProp, ElementPtr, DefaultElementPtr, ChildPropertyGatherTextFlags);
 				++ElementIndex;
 			}
 		}
diff --git a/Engine/Source/Runtime/CoreUObject/Public/UObject/UnrealType.h b/Engine/Source/Runtime/CoreUObject/Public/UObject/UnrealType.h
--- a/Engine/Source/Runtime/CoreUObject/Public/UObject/UnrealType.h
+++ b/Engine/Source/Runtime/CoreUObject/Public/UObject/UnrealType.h
@@ -3293,6 +3293,20 @@ public:
 		return Result;
 	}
 
+	/** Finds the associated pair from hash, rather than linearly searching */
+	uint8* FindMapPairPtrFromHash(const void* KeyPtr)
+	{
+		UProperty* LocalKeyPropForCapture = KeyProp;
+		int32 Index = Map->FindPairIndex(
+			KeyPtr,
+			MapLayout,
+			[LocalKeyPropForCapture](const void* ElementKey) { return LocalKeyPropForCapture->GetValueTypeHash(ElementKey); },
+			[LocalKeyPropForCapture](const void* A, const void* B) { return LocalKeyPropForCapture->Identical(A, B); }
+			);
+		uint8* Result = (Index >= 0) ? GetPairPtr(Index) : nullptr;
+		return Result;
+	}
+
 	/** Finds the associated value from hash, rather than linearly searching */
 	uint8* FindValueFromHash(const void* KeyPtr)
 	{
@@ -3829,6 +3843,14 @@ public:
 		);
 	}
 
+	/** Finds element pointer from hash, rather than linearly searching */
+	FORCEINLINE uint8* FindElementPtrFromHash(const void* ElementToFind)
+	{
+		const int32 Index = FindElementIndexFromHash(ElementToFind);
+		uint8* Result = (Index >= 0 ? GetElementPtr(Index) : nullptr);
+		return Result;
+	}
+
 	/** Adds the element to the set, returning true if the element was added, or false if the element was already present */
 	void AddElement(const void* ElementToAdd)
 	{
diff --git a/Engine/Source/Runtime/IOS/IOSAudio/Private/IOSAudioBuffer.cpp b/Engine/Source/Runtime/IOS/IOSAudio/Private/IOSAudioBuffer.cpp
--- a/Engine/Source/Runtime/IOS/IOSAudio/Private/IOSAudioBuffer.cpp
+++ b/Engine/Source/Runtime/IOS/IOSAudio/Private/IOSAudioBuffer.cpp
@@ -21,20 +21,31 @@
 	FIOSAudioSoundBuffer
  ------------------------------------------------------------------------------------*/
 
-FIOSAudioSoundBuffer::FIOSAudioSoundBuffer(FIOSAudioDevice* InAudioDevice, USoundWave* InWave, bool InStreaming):
+FIOSAudioSoundBuffer::FIOSAudioSoundBuffer(FIOSAudioDevice* InAudioDevice, USoundWave* InWave, bool InStreaming, bool InProcedural):
 	FSoundBuffer(InAudioDevice),
-	SampleData(NULL),
+	SampleData(nullptr),
 	BufferSize(0),
-	bStreaming(InStreaming)
+	DecompressionState(nullptr),
+	bStreaming(InStreaming),
+	bIsProcedural(InProcedural)
 {
-	DecompressionState = static_cast<FADPCMAudioInfo*>(InAudioDevice->CreateCompressedAudioInfo(InWave));
-
-	if (!ReadCompressedInfo(InWave))
+	if (!bIsProcedural)
 	{
-		return;
+		DecompressionState = static_cast<FADPCMAudioInfo*>(InAudioDevice->CreateCompressedAudioInfo(InWave));
+
+		if (!ReadCompressedInfo(InWave))
+		{
+			return;
+		}
+
+		SoundFormat = static_cast<ESoundFormat>(*DecompressionState->WaveInfo.pFormatTag);
 	}
+    else
+    {
+        SoundFormat = SoundFormat_LPCM;
+    }
+	
 	
-	SoundFormat = static_cast<ESoundFormat>(*DecompressionState->WaveInfo.pFormatTag);
 	SampleRate = InWave->GetSampleRateForCurrentPlatform();
 	NumChannels = InWave->NumChannels;
 	BufferSize = AudioCallbackFrameSize * sizeof(uint16) * NumChannels;
@@ -57,16 +68,16 @@ FIOSAudioSoundBuffer::~FIOSAudioSoundBuffer(void)
 		UE_LOG(LogIOSAudio, Fatal, TEXT("Can't free resource '%s' as it was allocated in permanent pool."), *ResourceName);
 	}
 	
-	if(SampleData != NULL)
+	if(SampleData != nullptr)
 	{
 		FMemory::Free(SampleData);
-		SampleData = NULL;
+		SampleData = nullptr;
 	}
 	
-	if(DecompressionState != NULL)
+	if(DecompressionState != nullptr)
 	{
 		delete DecompressionState;
-		DecompressionState = NULL;
+		DecompressionState = nullptr;
 	}
 }
 
@@ -87,7 +98,7 @@ int32 FIOSAudioSoundBuffer::GetSize(void)
 
 int32 FIOSAudioSoundBuffer::GetCurrentChunkIndex() const
 {
-	if(DecompressionState == NULL)
+	if(DecompressionState == nullptr)
 	{
 		return -1;
 	}
@@ -97,7 +108,7 @@ int32 FIOSAudioSoundBuffer::GetCurrentChunkIndex() const
 
 int32 FIOSAudioSoundBuffer::GetCurrentChunkOffset() const
 {
-	if(DecompressionState == NULL)
+	if(DecompressionState == nullptr)
 	{
 		return -1;
 	}
@@ -107,6 +118,8 @@ int32 FIOSAudioSoundBuffer::GetCurrentChunkOffset() const
 
 bool FIOSAudioSoundBuffer::ReadCompressedInfo(USoundWave* InWave)
 {
+	check(DecompressionState != nullptr)
+
 	FSoundQualityInfo QualityInfo = { 0 };
 	if(bStreaming)
 	{
@@ -129,7 +142,10 @@ bool FIOSAudioSoundBuffer::ReadCompressedData( uint8* Destination, int32 NumFram
 	{
 		return DecompressionState->StreamCompressedData(Destination, bLooping, RenderCallbackBufferSize);
 	}
-	return DecompressionState->ReadCompressedData(Destination, bLooping, RenderCallbackBufferSize);
+	else
+	{
+		return DecompressionState->ReadCompressedData(Destination, bLooping, RenderCallbackBufferSize);
+	}
 }
 
 FIOSAudioSoundBuffer* FIOSAudioSoundBuffer::Init(FIOSAudioDevice* IOSAudioDevice, USoundWave* InWave)
@@ -160,13 +176,16 @@ FIOSAudioSoundBuffer* FIOSAudioSoundBuffer::Init(FIOSAudioDevice* IOSAudioDevice
 						
 		case DTYPE_RealTime:
 			// Always create a new FIOSAudioSoundBuffer since positional information about the sound is kept track of in this object
-			Buffer = new FIOSAudioSoundBuffer(IOSAudioDevice, InWave, bIsStreaming);
+			Buffer = new FIOSAudioSoundBuffer(IOSAudioDevice, InWave, bIsStreaming, false);
 			break;
-			
+
+		case DTYPE_Procedural:
+			Buffer = new FIOSAudioSoundBuffer(IOSAudioDevice, InWave, bIsStreaming, true);
+			break;
+
 		case DTYPE_Native:
 		case DTYPE_Invalid:
 		case DTYPE_Preview:
-		case DTYPE_Procedural:
 		default:
 			// Invalid will be set if the wave cannot be played
 			UE_LOG( LogIOSAudio, Warning, TEXT("Init Buffer on unsupported sound type name = %s type = %d"), *InWave->GetName(), int32(InWave->DecompressionType));
diff --git a/Engine/Source/Runtime/IOS/IOSAudio/Private/IOSAudioSource.cpp b/Engine/Source/Runtime/IOS/IOSAudio/Private/IOSAudioSource.cpp
--- a/Engine/Source/Runtime/IOS/IOSAudio/Private/IOSAudioSource.cpp
+++ b/Engine/Source/Runtime/IOS/IOSAudio/Private/IOSAudioSource.cpp
@@ -464,21 +464,38 @@ OSStatus FIOSAudioSoundSource::IOSAudioRenderCallback(void* RefCon, AudioUnitRen
 		FMemory::Memzero(OutData, NumFrames * sizeof(AudioSampleType));
 		return -1;
 	}
-	
+
+    
+    
 	if(Channel == 0)
 	{
-		Source->IOSBuffer->RenderCallbackBufferSize = NumFrames * sizeof(uint16) * Source->IOSBuffer->DecompressionState->NumChannels;
-		
-		// Since StreamCompressedData returns interlaced samples we need to decompress all frames(samples) for all channels here so we don't end up decompressing multiple times
-		// Ensure we have enough memory to do this. If needed we could realloc here but that is bad practice inside the audio callback since it has a hard deadline
-		check(Source->IOSBuffer->RenderCallbackBufferSize <= Source->IOSBuffer->BufferSize)
-		
-		// Grab the interlaced channel data
-		Source->bChannel0Finished =
-			Source->IOSBuffer->ReadCompressedData(
+		// Grab the interleaved channel data
+		if (Source->IOSBuffer->bIsProcedural)
+		{
+            Source->IOSBuffer->RenderCallbackBufferSize = NumFrames * sizeof(uint16) * Source->IOSBuffer->NumChannels;
+            check(Source->IOSBuffer->RenderCallbackBufferSize <= Source->IOSBuffer->BufferSize);
+            
+            FMemory::Memzero(Source->IOSBuffer->SampleData, Source->IOSBuffer->RenderCallbackBufferSize);
+            
+			int32 DataSize = Source->WaveInstance->WaveData->GeneratePCMData(
+				(uint8*) Source->IOSBuffer->SampleData,
+				Source->IOSBuffer->RenderCallbackBufferSize / sizeof(int16));
+			Source->bChannel0Finished = DataSize <= 0;
+		}
+		else
+		{
+			Source->IOSBuffer->RenderCallbackBufferSize = NumFrames * sizeof(uint16) * Source->IOSBuffer->DecompressionState->NumChannels;
+
+			// Since StreamCompressedData returns interlaced samples we need to decompress all frames(samples) for all channels here so we don't end up decompressing multiple times
+			// Ensure we have enough memory to do this. If needed we could realloc here but that is bad practice inside the audio callback since it has a hard deadline
+			check(Source->IOSBuffer->RenderCallbackBufferSize <= Source->IOSBuffer->BufferSize);
+
+			Source->bChannel0Finished =
+				Source->IOSBuffer->ReadCompressedData(
 				(uint8*)Source->IOSBuffer->SampleData,
 				MONO_PCM_BUFFER_SAMPLES,
 				Source->WaveInstance->LoopingMode == LOOP_WithNotification || Source->WaveInstance->LoopingMode == LOOP_Forever);
+		}
 	}
     
     // If the channel count is higher than we've expected,
diff --git a/Engine/Source/Runtime/IOS/IOSAudio/Public/IOSAudioDevice.h b/Engine/Source/Runtime/IOS/IOSAudio/Public/IOSAudioDevice.h
--- a/Engine/Source/Runtime/IOS/IOSAudio/Public/IOSAudioDevice.h
+++ b/Engine/Source/Runtime/IOS/IOSAudio/Public/IOSAudioDevice.h
@@ -48,7 +48,7 @@ enum ESoundFormat
 class FIOSAudioSoundBuffer : public FSoundBuffer
 {
 public:
-	FIOSAudioSoundBuffer(FIOSAudioDevice* InAudioDevice, USoundWave* InWave, bool InStreaming);
+	FIOSAudioSoundBuffer(FIOSAudioDevice* InAudioDevice, USoundWave* InWave, bool InStreaming, bool InProcedural);
 	virtual ~FIOSAudioSoundBuffer(void);
 
 	/**
@@ -95,6 +95,7 @@ public:
 	/** Wrapper to assist in the bookkeeping of uncompressed data when streaming */
 	class FADPCMAudioInfo*		DecompressionState;
 	bool	bStreaming;
+	bool    bIsProcedural;
 };
 
 /**
@@ -200,6 +201,7 @@ public:
 	AUGraph GetAudioUnitGraph() const { return AudioUnitGraph; }
 	AUNode GetMixerNode() const { return MixerNode; }
 	AudioUnit GetMixerUnit() const { return MixerUnit; }
+	AudioUnit GetOutputUnit() const { return OutputUnit; }
 
 	/** Thread context management */
 	virtual void ResumeContext();
diff --git a/Engine/Source/Runtime/MediaAssets/Private/Assets/MediaPlayer.cpp b/Engine/Source/Runtime/MediaAssets/Private/Assets/MediaPlayer.cpp
--- a/Engine/Source/Runtime/MediaAssets/Private/Assets/MediaPlayer.cpp
+++ b/Engine/Source/Runtime/MediaAssets/Private/Assets/MediaPlayer.cpp
@@ -3,6 +3,7 @@
 #include "MediaPlayer.h"
 #include "MediaAssetsPrivate.h"
 
+#include "Engine/Engine.h"
 #include "IMediaClock.h"
 #include "IMediaControls.h"
 #include "IMediaModule.h"
@@ -10,6 +11,7 @@
 #include "IMediaPlayerFactory.h"
 #include "IMediaTicker.h"
 #include "IMediaTracks.h"
+#include "LatentActions.h"
 #include "MediaPlayerFacade.h"
 #include "MediaPlayerOptions.h"
 #include "Misc/App.h"
@@ -815,6 +817,8 @@ void UMediaPlayer::HandlePlayerMediaEvent(EMediaEvent Event)
 {
 	MediaEvent.Broadcast(Event);
 
+	bool bPlayOnOpen = false;
+
 	switch(Event)
 	{
 	case EMediaEvent::MediaClosed:
@@ -823,17 +827,39 @@ void UMediaPlayer::HandlePlayerMediaEvent(EMediaEvent Event)
 
 	case EMediaEvent::MediaOpened:
 		PlayerFacade->SetCacheWindow(CacheAhead, FApp::IsGame() ? CacheBehindGame : CacheBehind);
-		PlayerFacade->SetLooping(Loop && (Playlist->Num() == 1));
+		if (PlayerFacade->ActivePlayerOptions.IsSet() && PlayerFacade->ActivePlayerOptions->Loop != EMediaPlayerOptionBooleanOverride::UseMediaPlayerSetting)
+		{
+			PlayerFacade->SetLooping(PlayerFacade->ActivePlayerOptions->Loop == EMediaPlayerOptionBooleanOverride::Enabled);
+		}
+		else
+		{
+			PlayerFacade->SetLooping(Loop && (Playlist->Num() == 1));
+		}
 		PlayerFacade->SetViewField(HorizontalFieldOfView, VerticalFieldOfView, true);
 		PlayerFacade->SetViewOrientation(FQuat(ViewRotation), true);
 		PlayerFacade->TimeDelay = TimeDelay;
 
 		OnMediaOpened.Broadcast(PlayerFacade->GetUrl());
 
-		if (PlayOnOpen || PlayOnNext)
+		if (PlayerFacade->ActivePlayerOptions.IsSet() && PlayerFacade->ActivePlayerOptions->PlayOnOpen != EMediaPlayerOptionBooleanOverride::UseMediaPlayerSetting)
+		{
+			bPlayOnOpen = (PlayerFacade->ActivePlayerOptions->PlayOnOpen == EMediaPlayerOptionBooleanOverride::Enabled);
+		}
+		else
+		{
+			bPlayOnOpen = (PlayOnOpen || PlayOnNext);
+		}
+
+		if (bPlayOnOpen)
 		{
 			PlayOnNext = false;
-			Play();
+			if (Play())
+			{
+				if (PlayerFacade->ActivePlayerOptions.IsSet() && !PlayerFacade->ActivePlayerOptions->SeekTime.IsZero() && SupportsSeeking())
+				{
+					Seek(PlayerFacade->ActivePlayerOptions->SeekTime);
+				}
+			}
 		}
 		break;
 
@@ -875,3 +901,102 @@ void UMediaPlayer::HandlePlayerMediaEvent(EMediaEvent Event)
 		break;
 	}
 }
+
+class FLatentOpenMediaSourceAction : public FPendingLatentAction
+{
+public:
+	FName ExecutionFunction;
+	int32 OutputLink;
+	FWeakObjectPtr CallbackTarget;
+
+	TWeakObjectPtr<UMediaPlayer> MediaPlayer;
+	float TimeRemaining;
+	bool& OutSuccess;
+	bool bSawError;
+	FString URL;
+
+	FLatentOpenMediaSourceAction(const FLatentActionInfo& LatentInfo, UMediaPlayer* InMediaPlayer, UMediaSource* InMediaSource, const FMediaPlayerOptions& InOptions, bool& InSuccess)
+		: ExecutionFunction(LatentInfo.ExecutionFunction)
+		, OutputLink(LatentInfo.Linkage)
+		, CallbackTarget(LatentInfo.CallbackTarget)
+		, MediaPlayer(InMediaPlayer)
+		, TimeRemaining(10.0)
+		, OutSuccess(InSuccess)
+		, bSawError(false)
+	{
+		if (InMediaSource)
+		{
+			URL = InMediaSource->GetUrl();
+		}
+
+		if (!InMediaPlayer->OpenSourceWithOptions(InMediaSource, InOptions))
+		{
+			UE_LOG(LogMediaAssets, Verbose, TEXT("Latent Media: After initial open started"));
+			bSawError = true;
+		}
+	}
+
+	virtual void UpdateOperation(FLatentResponse& Response) override
+	{
+		if (bSawError || !MediaPlayer.IsValid() || MediaPlayer.IsStale() || MediaPlayer->HasError())
+		{
+			UE_LOG(LogMediaAssets, Warning, TEXT("Latent Media: Deleted or Error"));
+			OutSuccess = false;
+			Response.FinishAndTriggerIf(false, ExecutionFunction, OutputLink, CallbackTarget);
+			return;
+		}
+
+		if (MediaPlayer->IsReady())
+		{
+			UE_LOG(LogMediaAssets, Verbose, TEXT("Latent Media: IsReady"));
+			OutSuccess = true;
+			Response.FinishAndTriggerIf(true, ExecutionFunction, OutputLink, CallbackTarget);
+			return;
+		}
+		else if (!MediaPlayer->IsPreparing())
+		{
+			UE_LOG(LogMediaAssets, Verbose, TEXT("Latent Media: Invalid State"));
+			OutSuccess = false;
+			Response.FinishAndTriggerIf(false, ExecutionFunction, OutputLink, CallbackTarget);
+			return;
+		}
+		else
+		{
+			UE_LOG(LogMediaAssets, Verbose, TEXT("Latent Media: Waiting for open ..."));
+		}
+
+		// Timed out
+		TimeRemaining -= Response.ElapsedTime();
+		if (TimeRemaining <= 0.0f)
+		{
+			UE_LOG(LogMediaAssets, Verbose, TEXT("Latent Media: Timed out."));
+			OutSuccess = false;
+			Response.FinishAndTriggerIf(true, ExecutionFunction, OutputLink, CallbackTarget);
+			return;
+		}
+	}
+
+#if WITH_EDITOR
+	// Returns a human readable description of the latent operation's current state
+	virtual FString GetDescription() const override
+	{
+		return FString::Printf(TEXT("Opening Media: %s"), *URL);
+	}
+#endif
+};
+
+void UMediaPlayer::OpenSourceLatent(const UObject* WorldContextObject, FLatentActionInfo LatentInfo, UMediaSource* MediaSource, const FMediaPlayerOptions& Options, bool& bSuccess)
+{
+	bSuccess = false;
+
+	if (UWorld* World = GEngine->GetWorldFromContextObject(WorldContextObject, EGetWorldErrorMode::LogAndReturnNull))
+	{
+		FLatentActionManager& LatentManager = World->GetLatentActionManager();
+		if (LatentManager.FindExistingAction<FLatentOpenMediaSourceAction>(LatentInfo.CallbackTarget, LatentInfo.UUID) == nullptr)
+		{
+			FLatentOpenMediaSourceAction* NewAction = new FLatentOpenMediaSourceAction(LatentInfo, this, MediaSource, Options, bSuccess);
+			LatentManager.AddNewAction(LatentInfo.CallbackTarget, LatentInfo.UUID, NewAction);
+		}
+	}
+}
+
diff --git a/Engine/Source/Runtime/MediaAssets/Public/MediaPlayer.h b/Engine/Source/Runtime/MediaAssets/Public/MediaPlayer.h
--- a/Engine/Source/Runtime/MediaAssets/Public/MediaPlayer.h
+++ b/Engine/Source/Runtime/MediaAssets/Public/MediaPlayer.h
@@ -12,6 +12,7 @@
 #include "UObject/Object.h"
 #include "UObject/ScriptMacros.h"
 #include "Misc/Guid.h"
+#include "Engine/LatentActionManager.h"
 #include "MediaPlayerOptions.h"
 
 #include "MediaPlayer.generated.h"
@@ -569,6 +570,19 @@ public:
 	UFUNCTION(BlueprintCallable, Category="Media|MediaPlayer")
 	bool OpenSourceWithOptions(UMediaSource* MediaSource, const FMediaPlayerOptions& Options);
 
+	/**
+	 * Open the specified media source with options using a latent action.
+	 *
+	 * A result of true indicates that the player opened the media source successfully.
+	 *
+	 * @param MediaSource The media source to open.
+	 * @param Options The media player options to apply.
+	 * @param bSuccess  Player opened the media.
+	 * @see Close, OpenFile, OpenPlaylist, OpenPlaylistIndex, OpenUrl, Reopen
+	 */
+	UFUNCTION(BlueprintCallable, Category = "Media|MediaPlayer", meta = (Latent, LatentInfo = "LatentInfo", WorldContext = "WorldContextObject"))
+	void OpenSourceLatent(const UObject* WorldContextObject, struct FLatentActionInfo LatentInfo, UMediaSource* MediaSource, const FMediaPlayerOptions& Options, bool& bSuccess);
+
 	/**
 	 * Opens the specified media URL.
 	 *
diff --git a/Engine/Source/Runtime/MediaUtils/Public/MediaPlayerOptions.h b/Engine/Source/Runtime/MediaUtils/Public/MediaPlayerOptions.h
--- a/Engine/Source/Runtime/MediaUtils/Public/MediaPlayerOptions.h
+++ b/Engine/Source/Runtime/MediaUtils/Public/MediaPlayerOptions.h
@@ -8,6 +8,14 @@
 
 #include "MediaPlayerOptions.generated.h"
 
+UENUM(BlueprintType)
+enum class EMediaPlayerOptionBooleanOverride : uint8
+{
+	UseMediaPlayerSetting,
+	Enabled,
+	Disabled
+};
+
 USTRUCT(BlueprintType)
 struct FMediaPlayerTrackOptions
 {
@@ -51,10 +59,22 @@ struct FMediaPlayerOptions
 {
 	GENERATED_BODY()
 
-	FMediaPlayerOptions()
+	FMediaPlayerOptions() :
+		SeekTime(0),
+		PlayOnOpen(EMediaPlayerOptionBooleanOverride::UseMediaPlayerSetting),
+		Loop(EMediaPlayerOptionBooleanOverride::UseMediaPlayerSetting)
 	{
 	}
 
 	UPROPERTY(BlueprintReadWrite, Category = "Tracks")
 	FMediaPlayerTrackOptions Tracks;
+
+	UPROPERTY(BlueprintReadWrite, Category = "Misc")
+	FTimespan SeekTime;
+
+	UPROPERTY(BlueprintReadWrite, Category = "Misc")
+	EMediaPlayerOptionBooleanOverride PlayOnOpen;
+
+	UPROPERTY(BlueprintReadWrite, Category = "Misc")
+	EMediaPlayerOptionBooleanOverride Loop;
 };
