Remove hack for audio becoming desynced after garbage collection

diff --git a/Engine/Source/Runtime/MediaAssets/Private/Assets/MediaSoundComponent.cpp b/Engine/Source/Runtime/MediaAssets/Private/Assets/MediaSoundComponent.cpp
--- a/Engine/Source/Runtime/MediaAssets/Private/Assets/MediaSoundComponent.cpp
+++ b/Engine/Source/Runtime/MediaAssets/Private/Assets/MediaSoundComponent.cpp
@@ -29,14 +29,6 @@ FAutoConsoleVariableRef CVarSyncAudioAfterDropouts(
 	TEXT("0: Not Enabled, 1: Enabled"),
 	ECVF_Default);
 
-static int32 FlushBackloggedAudioCVar = 0;
-FAutoConsoleVariableRef CVarFlushBackloggedAudio(
-	TEXT("m.FlushBackloggedAudio"),
-	FlushBackloggedAudioCVar,
-	TEXT("Flush backlogged audio samples when max value reached.\n")
-	TEXT("0: Not Enabled, 4 or higher: Max queued samples"),
-	ECVF_Default);
-
 DECLARE_FLOAT_COUNTER_STAT(TEXT("MediaUtils MediaSoundComponent Sync"), STAT_MediaUtils_MediaSoundComponentSync, STATGROUP_Media);
 DECLARE_FLOAT_COUNTER_STAT(TEXT("MediaUtils MediaSoundComponent SampleTime"), STAT_MediaUtils_MediaSoundComponentSampleTime, STATGROUP_Media);
 DECLARE_DWORD_COUNTER_STAT(TEXT("MediaUtils MediaSoundComponent Queued"), STAT_Media_SoundCompQueued, STATGROUP_Media);
@@ -481,20 +473,6 @@ int32 UMediaSoundComponent::OnGenerateAudio(float* OutAudio, int32 NumSamples)
 			}
 		}
 
-#if PLATFORM_WINDOWS
-		// Garbage collection and other operations can cause calls to OnGenerateAudio() to be spaced too far apart (no audio being processed)
-		// These gaps are not tracked leading to a backlogged buffer that becomes further and further out of sync
-		if (FlushBackloggedAudioCVar > 3 && PinnedSampleQueue->Num() > FlushBackloggedAudioCVar)
-		{
-			const int32 NumQueued = PinnedSampleQueue->Num();
-			for (int32 i = 0; i < NumQueued - 3; ++i)
-			{
-				TSharedPtr<IMediaAudioSample, ESPMode::ThreadSafe> NextSample;
-				PinnedSampleQueue->Dequeue(NextSample);
-			}
-		}
-#endif
-
 		SET_FLOAT_STAT(STAT_MediaUtils_MediaSoundComponentSync, FMath::Abs((Time - OutTime).GetTotalMilliseconds()));
 		SET_FLOAT_STAT(STAT_MediaUtils_MediaSoundComponentSampleTime, OutTime.GetTotalMilliseconds());
 		SET_DWORD_STAT(STAT_Media_SoundCompQueued, PinnedSampleQueue->Num());
