- Fix MediaTexture bug when switching player. Material using it were using new guids before it was effectively taken into account in the MediaTextureResource.
- Fix in MediaTextureResource where it could use a sample of the future but not remove it from the queue

diff --git a/Engine/Source/Runtime/MediaAssets/Private/Assets/MediaTexture.cpp b/Engine/Source/Runtime/MediaAssets/Private/Assets/MediaTexture.cpp
--- a/Engine/Source/Runtime/MediaAssets/Private/Assets/MediaTexture.cpp
+++ b/Engine/Source/Runtime/MediaAssets/Private/Assets/MediaTexture.cpp
@@ -158,9 +158,17 @@ float UMediaTexture::GetSurfaceHeight() const
 
 FGuid UMediaTexture::GetExternalTextureGuid() const
 {
-	return CurrentGuid;
+	FScopeLock Lock(&CriticalSection);
+	return CurrentRenderedGuid;
 }
 
+void UMediaTexture::SetRenderedExternalTextureGuid(const FGuid& InNewGuid)
+{
+	check(IsInRenderingThread());
+
+	FScopeLock Lock(&CriticalSection);
+	CurrentRenderedGuid = InNewGuid;
+}
 
 /* UObject interface
  *****************************************************************************/
@@ -179,13 +187,14 @@ void UMediaTexture::BeginDestroy()
 		ClockSink.Reset();
 	}
 
-	if (CurrentGuid.IsValid())
+	//Unregister the last rendered Guid
+	const FGuid LastRendered = GetExternalTextureGuid();
+	if (LastRendered.IsValid())
 	{
-		FGuid Guid = CurrentGuid;
 		ENQUEUE_RENDER_COMMAND(MediaTextureUnregisterGuid)(
-			[Guid](FRHICommandList& RHICmdList)
+			[LastRendered](FRHICommandList& RHICmdList)
 			{
-				FExternalTextureRegistry::Get().UnregisterExternalTexture(Guid);
+				FExternalTextureRegistry::Get().UnregisterExternalTexture(LastRendered);
 			});
 	}
 
diff --git a/Engine/Source/Runtime/MediaAssets/Private/Misc/MediaTextureResource.cpp b/Engine/Source/Runtime/MediaAssets/Private/Misc/MediaTextureResource.cpp
--- a/Engine/Source/Runtime/MediaAssets/Private/Misc/MediaTextureResource.cpp
+++ b/Engine/Source/Runtime/MediaAssets/Private/Misc/MediaTextureResource.cpp
@@ -185,13 +185,14 @@ void FMediaTextureResource::Render(const FRenderParams& Params)
 	if (SampleSource.IsValid())
 	{
 		// get the most current sample to be rendered
+		TSharedPtr<IMediaTextureSample, ESPMode::ThreadSafe> TestSample;
 		TSharedPtr<IMediaTextureSample, ESPMode::ThreadSafe> Sample;
 		bool UseSample = false;
 		
-		while (SampleSource->Peek(Sample) && Sample.IsValid())
+		while (SampleSource->Peek(TestSample) && TestSample.IsValid())
 		{
-			const FTimespan StartTime = Sample->GetTime();
-			const FTimespan EndTime = StartTime + Sample->GetDuration();
+			const FTimespan StartTime = TestSample->GetTime();
+			const FTimespan EndTime = StartTime + TestSample->GetDuration();
 
 			if ((Params.Rate >= 0.0f) && (Params.Time < StartTime))
 			{
@@ -332,6 +333,9 @@ void FMediaTextureResource::Render(const FRenderParams& Params)
 			FExternalTextureRegistry::Get().UnregisterExternalTexture(Params.PreviousGuid);
 		}
 	}
+	
+	//Update usable Guid for the RenderThread
+	Owner.SetRenderedExternalTextureGuid(Params.CurrentGuid);
 }
 
 
diff --git a/Engine/Source/Runtime/MediaAssets/Public/MediaTexture.h b/Engine/Source/Runtime/MediaAssets/Public/MediaTexture.h
--- a/Engine/Source/Runtime/MediaAssets/Public/MediaTexture.h
+++ b/Engine/Source/Runtime/MediaAssets/Public/MediaTexture.h
@@ -128,6 +128,7 @@ public:
 	virtual float GetSurfaceWidth() const override;
 	virtual float GetSurfaceHeight() const override;
 	virtual FGuid GetExternalTextureGuid() const override;
+	void SetRenderedExternalTextureGuid(const FGuid& InNewGuid);
 
 public:
 
@@ -175,7 +176,10 @@ private:
 	TSharedPtr<FMediaTextureClockSink, ESPMode::ThreadSafe> ClockSink;
 
 	/** The default external texture GUID if no media player is assigned. */
-	/*TAtomic<*/FGuid/*>*/ CurrentGuid;
+	FGuid CurrentGuid;
+
+	/** The last Guid that was rendered and registered in the render command*/
+	FGuid CurrentRenderedGuid;
 
 	/** The player that is currently associated with this texture. */
 	TWeakObjectPtr<UMediaPlayer> CurrentPlayer;
@@ -197,4 +201,7 @@ private:
 
 	/** Current size of the resource (in bytes).*/
 	SIZE_T Size;
+
+	/** Critical section to protect last rendered guid since it can be read from anywhere. */
+	mutable FCriticalSection CriticalSection;
 };
