Temporary workaround for audio starving handling when streaming, per-platform.
#rnx

diff --git a/Engine/Source/Runtime/MediaAssets/Public/MediaSoundComponent.h b/Engine/Source/Runtime/MediaAssets/Public/MediaSoundComponent.h
--- a/Engine/Source/Runtime/MediaAssets/Public/MediaSoundComponent.h
+++ b/Engine/Source/Runtime/MediaAssets/Public/MediaSoundComponent.h
@@ -183,6 +183,14 @@ public:
 	void SetDefaultMediaPlayer(UMediaPlayer* NewMediaPlayer);
 #endif
 
+	/**
+	 * Allows tweaking the behavior for when the audio starves
+	 */
+	void SetSyncAudioAfterDropouts(bool bInSyncAudioAfterDropouts)
+	{
+		bSyncAudioAfterDropouts = bInSyncAudioAfterDropouts;
+	}
+
 public:
 
 	//~ TAttenuatedComponentVisualizer interface
