Support caching more than 1 sample if needed by platform.

diff --git a/Engine/Source/Runtime/MediaAssets/Private/Misc/MediaTextureResource.cpp b/Engine/Source/Runtime/MediaAssets/Private/Misc/MediaTextureResource.cpp
--- a/Engine/Source/Runtime/MediaAssets/Private/Misc/MediaTextureResource.cpp
+++ b/Engine/Source/Runtime/MediaAssets/Private/Misc/MediaTextureResource.cpp
@@ -33,6 +33,12 @@ DECLARE_CYCLE_STAT(TEXT("MediaAssets MediaTextureResource Render"), STAT_MediaAs
 DECLARE_FLOAT_COUNTER_STAT(TEXT("MediaAssets MediaTextureResource Sample"), STAT_MediaUtils_TextureSampleTime, STATGROUP_Media);
 
 
+static int32 CachedSamplesQueueDepth = 1;
+static FAutoConsoleVariableRef CVarEncoderSaveVideoToFile(
+	TEXT("media.CachedSamplesQueueDepth"),
+	CachedSamplesQueueDepth,
+	TEXT("How many frames to hold samples before release (default = 1)."),
+	ECVF_Default);
 
 /* Local helpers
  *****************************************************************************/
@@ -163,7 +169,13 @@ FMediaTextureResource::FMediaTextureResource(UMediaTexture& InOwner, FIntPoint&
 	, Owner(InOwner)
 	, OwnerDim(InOwnerDim)
 	, OwnerSize(InOwnerSize)
-{ }
+{
+	// preset the CachedSamples. makes things easier in the long run
+	if (CachedSamplesQueueDepth > 0)
+	{
+		CachedSamples.AddDefaulted(CachedSamplesQueueDepth);
+	}
+}
 
 
 /* FMediaTextureResource interface
@@ -173,7 +185,7 @@ void FMediaTextureResource::Render(const FRenderParams& Params)
 {
 	check(IsInRenderingThread());
 
-	CachedSample.Reset();
+	CycleCachedSamples();
 
 	SCOPE_CYCLE_COUNTER(STAT_MediaAssets_MediaTextureResourceRender);
 
@@ -301,8 +313,11 @@ void FMediaTextureResource::Render(const FRenderParams& Params)
 #endif
 
 		// We're not done with `Sample` as rendering is asynchronous. 
-		// Hold a reference in a member to postpone recycling `Sample` till the next call
-		CachedSample = Sample;
+		// Hold a reference in a member to postpone recycling `Sample` until safe to release
+		if (CachedSamplesQueueDepth > 0)
+		{
+			CachedSamples[0] = Sample;
+		}
 	}
 	else if (Params.CanClear)
 	{
@@ -338,6 +353,16 @@ void FMediaTextureResource::Render(const FRenderParams& Params)
 	Owner.SetRenderedExternalTextureGuid(Params.CurrentGuid);
 }
 
+void FMediaTextureResource::CycleCachedSamples()
+{
+	int32 CachedSampleToChange = CachedSamplesQueueDepth - 1;
+
+	while (CachedSampleToChange > 0)
+	{
+		CachedSamples[CachedSampleToChange] = CachedSamples[CachedSampleToChange - 1];
+		CachedSampleToChange--;
+	}
+}
 
 /* FRenderTarget interface
  *****************************************************************************/
diff --git a/Engine/Source/Runtime/MediaAssets/Private/Misc/MediaTextureResource.h b/Engine/Source/Runtime/MediaAssets/Private/Misc/MediaTextureResource.h
--- a/Engine/Source/Runtime/MediaAssets/Private/Misc/MediaTextureResource.h
+++ b/Engine/Source/Runtime/MediaAssets/Private/Misc/MediaTextureResource.h
@@ -151,6 +151,11 @@ protected:
 
 private:
 
+	/**
+	* Cycle our cached samples
+	*/
+	void CycleCachedSamples();
+
 	/** Whether the texture has been cleared. */
 	bool Cleared;
 
@@ -179,5 +184,5 @@ private:
 	TWeakPtr<FMediaPlayerFacade, ESPMode::ThreadSafe> PlayerFacadePtr;
 
 	/** cached media sample to postpone releasing it until the next sample rendering as it can get overwritten due to asynchronous rendering */
-	TSharedPtr<IMediaTextureSample, ESPMode::ThreadSafe> CachedSample;
+	TArray<TSharedPtr<IMediaTextureSample, ESPMode::ThreadSafe>> CachedSamples;
 };
