UMediaSoundComponent now queries the used IMediaPlayer instance for how it should handle audio dropouts.

#rnx

diff --git a/Engine/Source/Runtime/Media/Public/IMediaPlayer.h b/Engine/Source/Runtime/Media/Public/IMediaPlayer.h
--- a/Engine/Source/Runtime/Media/Public/IMediaPlayer.h
+++ b/Engine/Source/Runtime/Media/Public/IMediaPlayer.h
@@ -166,6 +166,22 @@ public:
 	 */
 	virtual bool Open(const TSharedRef<FArchive, ESPMode::ThreadSafe>& Archive, const FString& OriginalUrl, const IMediaOptions* Options) = 0;
 
+	/**
+	* Specifies how audio starving should be handled
+	* 
+	* This affects the starving behavior of the used UMediaSoundComponent component.
+	* When true, it will skip audio to keep the time as if no starving occurred.
+	* If false, it will simply wait and continue without skipping when it gets more audio.
+	*/
+	virtual bool RequiresAudioSyncAfterDropouts() const
+	{
+#if PLATFORM_PS4 || PLATFORM_SWITCH || PLATFORM_XBOXONE
+		return true;
+#else
+		return false;
+#endif
+	}
+
 public:
 
 	//~ The following methods are optional
diff --git a/Engine/Source/Runtime/MediaAssets/Private/Assets/MediaSoundComponent.cpp b/Engine/Source/Runtime/MediaAssets/Private/Assets/MediaSoundComponent.cpp
--- a/Engine/Source/Runtime/MediaAssets/Private/Assets/MediaSoundComponent.cpp
+++ b/Engine/Source/Runtime/MediaAssets/Private/Assets/MediaSoundComponent.cpp
@@ -157,6 +157,19 @@ void UMediaSoundComponent::UpdatePlayer()
 		CurrentPlayerFacade = PlayerFacade;
 	}
 
+	//
+	// Some players might require bSyncAudioAfterDropouts to be forced to true or false
+	if (PlayerFacade->GetPlayer())
+	{
+		bool bNewAudioSyncAfterDropouts = PlayerFacade->GetPlayer()->RequiresAudioSyncAfterDropouts();
+		if (bNewAudioSyncAfterDropouts != bSyncAudioAfterDropouts)
+		{
+			FScopeLock Lock(&CriticalSection);
+			bSyncAudioAfterDropouts = bNewAudioSyncAfterDropouts;
+			FrameSyncOffset = 0;
+		}
+	}
+
 	// caching play rate and time for audio thread (eventual consistency is sufficient)
 	CachedRate = PlayerFacade->GetRate();
 	CachedTime = PlayerFacade->GetTime();
