Removing UMediaSoundComponent::SetSyncAudioAfterDropouts. That was an hack, and we now hopefully have a better fix (CL 6143134)

#rnx

diff --git a/Engine/Source/Runtime/MediaAssets/Public/MediaSoundComponent.h b/Engine/Source/Runtime/MediaAssets/Public/MediaSoundComponent.h
--- a/Engine/Source/Runtime/MediaAssets/Public/MediaSoundComponent.h
+++ b/Engine/Source/Runtime/MediaAssets/Public/MediaSoundComponent.h
@@ -183,14 +183,6 @@ public:
 	void SetDefaultMediaPlayer(UMediaPlayer* NewMediaPlayer);
 #endif
 
-	/**
-	 * Allows tweaking the behavior for when the audio starves
-	 */
-	void SetSyncAudioAfterDropouts(bool bInSyncAudioAfterDropouts)
-	{
-		bSyncAudioAfterDropouts = bInSyncAudioAfterDropouts;
-	}
-
 public:
 
 	//~ TAttenuatedComponentVisualizer interface
