Fix incorrect texture gamma when running with -featureleveles2 -game.
Fix incorrect media texture sRGB conversion with -featureleveles2 -game

diff --git a/Engine/Source/Runtime/MediaAssets/Private/Misc/MediaTextureResource.cpp b/Engine/Source/Runtime/MediaAssets/Private/Misc/MediaTextureResource.cpp
--- a/Engine/Source/Runtime/MediaAssets/Private/Misc/MediaTextureResource.cpp
+++ b/Engine/Source/Runtime/MediaAssets/Private/Misc/MediaTextureResource.cpp
@@ -618,6 +618,15 @@ void FMediaTextureResource::ConvertSample(const TSharedPtr<IMediaTextureSample,
 				YUVOffset = MediaShaders::YUVOffset10bits;
 			}
 
+			bool bIsSampleOutputSrgb = Sample->IsOutputSrgb();
+			if (GMaxRHIFeatureLevel == ERHIFeatureLevel::ES2 && IsSimulatedPlatform(GMaxRHIShaderPlatform) )
+			{
+				// simulated ES2 has no HW support for sRGB, all external textures are assumed to be in sRGB form. 
+				// see FHLSLMaterialTranslator::TextureSample(), SAMPLERTYPE_External case.
+				// We must not convert to linear for the ES2 case.
+				bIsSampleOutputSrgb = false;
+			}
+
 			switch (Sample->GetFormat())
 			{
 			case EMediaTextureSampleFormat::CharAYUV:
@@ -625,7 +634,7 @@ void FMediaTextureResource::ConvertSample(const TSharedPtr<IMediaTextureSample,
 				TShaderMapRef<FAYUVConvertPS> ConvertShader(ShaderMap);
 				GraphicsPSOInit.BoundShaderState.PixelShaderRHI = GETSAFERHISHADER_PIXEL(*ConvertShader);
 				SetGraphicsPipelineState(CommandList, GraphicsPSOInit);
-				ConvertShader->SetParameters(CommandList, InputTexture, YUVToRGBMatrix, YUVOffset, Sample->IsOutputSrgb());
+				ConvertShader->SetParameters(CommandList, InputTexture, YUVToRGBMatrix, YUVOffset, bIsSampleOutputSrgb);
 			}
 			break;
 
@@ -634,7 +643,7 @@ void FMediaTextureResource::ConvertSample(const TSharedPtr<IMediaTextureSample,
 				TShaderMapRef<FBMPConvertPS> ConvertShader(ShaderMap);
 				GraphicsPSOInit.BoundShaderState.PixelShaderRHI = GETSAFERHISHADER_PIXEL(*ConvertShader);
 				SetGraphicsPipelineState(CommandList, GraphicsPSOInit);
-				ConvertShader->SetParameters(CommandList, InputTexture, OutputDim, Sample->IsOutputSrgb() && !SrgbOutput);
+				ConvertShader->SetParameters(CommandList, InputTexture, OutputDim, bIsSampleOutputSrgb && !SrgbOutput);
 			}
 			break;
 
@@ -643,7 +652,7 @@ void FMediaTextureResource::ConvertSample(const TSharedPtr<IMediaTextureSample,
 				TShaderMapRef<FNV12ConvertPS> ConvertShader(ShaderMap);
 				GraphicsPSOInit.BoundShaderState.PixelShaderRHI = GETSAFERHISHADER_PIXEL(*ConvertShader);
 				SetGraphicsPipelineState(CommandList, GraphicsPSOInit);
-				ConvertShader->SetParameters(CommandList, InputTexture, OutputDim, YUVToRGBMatrix, YUVOffset, Sample->IsOutputSrgb());
+				ConvertShader->SetParameters(CommandList, InputTexture, OutputDim, YUVToRGBMatrix, YUVOffset, bIsSampleOutputSrgb);
 			}
 			break;
 
@@ -652,7 +661,7 @@ void FMediaTextureResource::ConvertSample(const TSharedPtr<IMediaTextureSample,
 				TShaderMapRef<FNV21ConvertPS> ConvertShader(ShaderMap);
 				GraphicsPSOInit.BoundShaderState.PixelShaderRHI = GETSAFERHISHADER_PIXEL(*ConvertShader);
 				SetGraphicsPipelineState(CommandList, GraphicsPSOInit);
-				ConvertShader->SetParameters(CommandList, InputTexture, OutputDim, YUVToRGBMatrix, YUVOffset, Sample->IsOutputSrgb());
+				ConvertShader->SetParameters(CommandList, InputTexture, OutputDim, YUVToRGBMatrix, YUVOffset, bIsSampleOutputSrgb);
 			}
 			break;
 
@@ -661,7 +670,7 @@ void FMediaTextureResource::ConvertSample(const TSharedPtr<IMediaTextureSample,
 				TShaderMapRef<FUYVYConvertPS> ConvertShader(ShaderMap);
 				GraphicsPSOInit.BoundShaderState.PixelShaderRHI = GETSAFERHISHADER_PIXEL(*ConvertShader);
 				SetGraphicsPipelineState(CommandList, GraphicsPSOInit);
-				ConvertShader->SetParameters(CommandList, InputTexture, YUVToRGBMatrix, YUVOffset, Sample->IsOutputSrgb());
+				ConvertShader->SetParameters(CommandList, InputTexture, YUVToRGBMatrix, YUVOffset, bIsSampleOutputSrgb);
 			}
 			break;
 
@@ -670,7 +679,7 @@ void FMediaTextureResource::ConvertSample(const TSharedPtr<IMediaTextureSample,
 				TShaderMapRef<FYUY2ConvertPS> ConvertShader(ShaderMap);
 				GraphicsPSOInit.BoundShaderState.PixelShaderRHI = GETSAFERHISHADER_PIXEL(*ConvertShader);
 				SetGraphicsPipelineState(CommandList, GraphicsPSOInit);
-				ConvertShader->SetParameters(CommandList, InputTexture, OutputDim, YUVToRGBMatrix, YUVOffset, Sample->IsOutputSrgb());
+				ConvertShader->SetParameters(CommandList, InputTexture, OutputDim, YUVToRGBMatrix, YUVOffset, bIsSampleOutputSrgb);
 			}
 			break;
 
@@ -679,7 +688,7 @@ void FMediaTextureResource::ConvertSample(const TSharedPtr<IMediaTextureSample,
 				TShaderMapRef<FYVYUConvertPS> ConvertShader(ShaderMap);
 				GraphicsPSOInit.BoundShaderState.PixelShaderRHI = GETSAFERHISHADER_PIXEL(*ConvertShader);
 				SetGraphicsPipelineState(CommandList, GraphicsPSOInit);
-				ConvertShader->SetParameters(CommandList, InputTexture, YUVToRGBMatrix, YUVOffset, Sample->IsOutputSrgb());
+				ConvertShader->SetParameters(CommandList, InputTexture, YUVToRGBMatrix, YUVOffset, bIsSampleOutputSrgb);
 			}
 			break;
 
@@ -688,7 +697,7 @@ void FMediaTextureResource::ConvertSample(const TSharedPtr<IMediaTextureSample,
 				TShaderMapRef<FYUVv210ConvertPS> ConvertShader(ShaderMap);
 				GraphicsPSOInit.BoundShaderState.PixelShaderRHI = GETSAFERHISHADER_PIXEL(*ConvertShader);
 				SetGraphicsPipelineState(CommandList, GraphicsPSOInit);
-				ConvertShader->SetParameters(CommandList, InputTexture, OutputDim, YUVToRGBMatrix, YUVOffset, Sample->IsOutputSrgb());
+				ConvertShader->SetParameters(CommandList, InputTexture, OutputDim, YUVToRGBMatrix, YUVOffset, bIsSampleOutputSrgb);
 			}
 			break;
 
@@ -697,7 +706,7 @@ void FMediaTextureResource::ConvertSample(const TSharedPtr<IMediaTextureSample,
 				TShaderMapRef<FRGBConvertPS> ConvertShader(ShaderMap);
 				GraphicsPSOInit.BoundShaderState.PixelShaderRHI = GETSAFERHISHADER_PIXEL(*ConvertShader);
 				SetGraphicsPipelineState(CommandList, GraphicsPSOInit);
-				ConvertShader->SetParameters(CommandList, InputTexture, OutputDim, Sample->IsOutputSrgb());
+				ConvertShader->SetParameters(CommandList, InputTexture, OutputDim, bIsSampleOutputSrgb);
 			}
 			break;
 
diff --git a/Engine/Source/Runtime/Renderer/Private/SceneRendering.cpp b/Engine/Source/Runtime/Renderer/Private/SceneRendering.cpp
--- a/Engine/Source/Runtime/Renderer/Private/SceneRendering.cpp
+++ b/Engine/Source/Runtime/Renderer/Private/SceneRendering.cpp
@@ -1391,8 +1391,11 @@ void FViewInfo::SetupUniformBufferParameters(
 	// Make sure there's no padding since we're going to cast to FVector4*
 	checkSlow(sizeof(ViewUniformShaderParameters.SkyIrradianceEnvironmentMap) == sizeof(FVector4)* 7);
 	SetupSkyIrradianceEnvironmentMapConstants((FVector4*)&ViewUniformShaderParameters.SkyIrradianceEnvironmentMap);
-
-	ViewUniformShaderParameters.MobilePreviewMode = Scene == nullptr ? 0.0f: (IsSimulatedPlatform(Scene->GetShaderPlatform()) ? 1.0f : 0.0f);
+	
+	ViewUniformShaderParameters.MobilePreviewMode =
+		(GIsEditor &&
+		(RHIFeatureLevel == ERHIFeatureLevel::ES2 || RHIFeatureLevel == ERHIFeatureLevel::ES3_1) &&
+		GMaxRHIFeatureLevel > ERHIFeatureLevel::ES3_1) ? 1.0f : 0.0f;
 
 	// Padding between the left and right eye may be introduced by an HMD, which instanced stereo needs to account for.
 	if ((StereoPass != eSSP_FULL) && (Family->Views.Num() > 1))
diff --git a/Engine/Source/Runtime/SlateRHIRenderer/Private/SlateRHIRenderingPolicy.cpp b/Engine/Source/Runtime/SlateRHIRenderer/Private/SlateRHIRenderingPolicy.cpp
--- a/Engine/Source/Runtime/SlateRHIRenderer/Private/SlateRHIRenderingPolicy.cpp
+++ b/Engine/Source/Runtime/SlateRHIRenderer/Private/SlateRHIRenderingPolicy.cpp
@@ -291,8 +291,10 @@ static FSceneView* CreateSceneView( FSceneViewFamilyContext* ViewFamilyContext,
 
 	ERHIFeatureLevel::Type RHIFeatureLevel = View->GetFeatureLevel();
 
-	ViewUniformShaderParameters.MobilePreviewMode = IsSimulatedPlatform(View->GetShaderPlatform()) ? 1.0f : 0.0f;
-
+	ViewUniformShaderParameters.MobilePreviewMode =
+		(GIsEditor &&
+		(RHIFeatureLevel == ERHIFeatureLevel::ES2 || RHIFeatureLevel == ERHIFeatureLevel::ES3_1) &&
+		GMaxRHIFeatureLevel > ERHIFeatureLevel::ES3_1) ? 1.0f : 0.0f;
 
 	UpdateNoiseTextureParameters(ViewUniformShaderParameters);
 
