-Adding an additional nullptr check for the cooking target when the BaseMediaSource gets serialized.

diff --git a/Engine/Source/Runtime/MediaAssets/Private/Assets/BaseMediaSource.cpp b/Engine/Source/Runtime/MediaAssets/Private/Assets/BaseMediaSource.cpp
--- a/Engine/Source/Runtime/MediaAssets/Private/Assets/BaseMediaSource.cpp
+++ b/Engine/Source/Runtime/MediaAssets/Private/Assets/BaseMediaSource.cpp
@@ -54,7 +54,7 @@ void UBaseMediaSource::Serialize(FArchive& Ar)
 #if WITH_EDITORONLY_DATA
 		if (Ar.IsFilterEditorOnly())
 		{
-			if (Ar.IsSaving())
+			if (Ar.IsSaving() && (Ar.CookingTarget() != nullptr))
 			{
 				const FName* PlatformPlayerName = PlatformPlayerNames.Find(Ar.CookingTarget()->IniPlatformName());
 				PlayerName = (PlatformPlayerName != nullptr) ? *PlatformPlayerName : NAME_None;
