MediaFramework: Added calls for setting media options on a media source.

diff --git a/Engine/Source/Runtime/MediaAssets/Private/Assets/MediaPlayer.cpp b/Engine/Source/Runtime/MediaAssets/Private/Assets/MediaPlayer.cpp
--- a/Engine/Source/Runtime/MediaAssets/Private/Assets/MediaPlayer.cpp
+++ b/Engine/Source/Runtime/MediaAssets/Private/Assets/MediaPlayer.cpp
@@ -617,6 +617,13 @@ bool UMediaPlayer::SetLooping(bool Looping)
 }
 
 
+void UMediaPlayer::SetMediaOptions(const UMediaSource* Options)
+{
+	UE_LOG(LogMediaAssets, Verbose, TEXT("%s.SetMediaOptions"), *GetFName().ToString());
+	PlayerFacade->SetMediaOptions(Options);
+}
+
+
 bool UMediaPlayer::SetRate(float Rate)
 {
 	UE_LOG(LogMediaAssets, Verbose, TEXT("%s.SetRate %f"), *GetFName().ToString(), Rate);
diff --git a/Engine/Source/Runtime/MediaAssets/Private/Assets/MediaSource.cpp b/Engine/Source/Runtime/MediaAssets/Private/Assets/MediaSource.cpp
--- a/Engine/Source/Runtime/MediaAssets/Private/Assets/MediaSource.cpp
+++ b/Engine/Source/Runtime/MediaAssets/Private/Assets/MediaSource.cpp
@@ -2,7 +2,6 @@
 
 #include "MediaSource.h"
 
-
 /* IMediaOptions interface
  *****************************************************************************/
 
@@ -14,25 +13,57 @@ FName UMediaSource::GetDesiredPlayerName() const
 
 bool UMediaSource::GetMediaOption(const FName& Key, bool DefaultValue) const
 {
-	return DefaultValue;
+	const FVariant* Variant = GetMediaOption(Key);
+	if ((Variant != nullptr) && (Variant->GetType() == EVariantTypes::Bool))
+	{
+		return Variant->GetValue<bool>();
+	}
+	else
+	{
+		return DefaultValue;
+	}
 }
 
 
 double UMediaSource::GetMediaOption(const FName& Key, double DefaultValue) const
 {
-	return DefaultValue;
+	const FVariant* Variant = GetMediaOption(Key);
+	if ((Variant != nullptr) && (Variant->GetType() == EVariantTypes::Double))
+	{
+		return Variant->GetValue<double>();
+	}
+	else
+	{
+		return DefaultValue;
+	}
 }
 
 
 int64 UMediaSource::GetMediaOption(const FName& Key, int64 DefaultValue) const
 {
-	return DefaultValue;
+	const FVariant* Variant = GetMediaOption(Key);
+	if ((Variant != nullptr) && (Variant->GetType() == EVariantTypes::Int64))
+	{
+		return Variant->GetValue<int64>();
+	}
+	else
+	{
+		return DefaultValue;
+	}
 }
 
 
 FString UMediaSource::GetMediaOption(const FName& Key, const FString& DefaultValue) const
 {
-	return DefaultValue;
+	const FVariant* Variant = GetMediaOption(Key);
+	if ((Variant != nullptr) && (Variant->GetType() == EVariantTypes::String))
+	{
+		return Variant->GetValue<FString>();
+	}
+	else
+	{
+		return DefaultValue;
+	}
 }
 
 
@@ -44,5 +75,53 @@ FText UMediaSource::GetMediaOption(const FName& Key, const FText& DefaultValue)
 
 bool UMediaSource::HasMediaOption(const FName& Key) const
 {
-	return false;
+	return MediaOptionsMap.Contains(Key);
+}
+
+
+void UMediaSource::SetMediaOptionBool(const FName& Key, bool Value)
+{
+	FVariant Variant(Value);
+	SetMediaOption(Key, Variant);
 }
+
+
+void UMediaSource::SetMediaOptionFloat(const FName& Key, float Value)
+{
+	SetMediaOptionDouble(Key, (double)Value);
+}
+
+
+void UMediaSource::SetMediaOptionDouble(const FName& Key, double Value)
+{
+	FVariant Variant(Value);
+	SetMediaOption(Key, Variant);
+}
+
+
+void UMediaSource::SetMediaOptionInt64(const FName& Key, int64 Value)
+{
+	FVariant Variant(Value);
+	SetMediaOption(Key, Variant);
+}
+
+
+void UMediaSource::SetMediaOptionString(const FName& Key, const FString& Value)
+{
+	FVariant Variant(Value);
+	SetMediaOption(Key, Variant);
+}
+
+
+const FVariant* UMediaSource::GetMediaOption(const FName& Key) const
+{
+	return MediaOptionsMap.Find(Key);
+}
+
+
+void UMediaSource::SetMediaOption(const FName& Key, FVariant& Value)
+{
+	MediaOptionsMap.Emplace(Key, Value);
+}
+
+
diff --git a/Engine/Source/Runtime/MediaAssets/Public/MediaPlayer.h b/Engine/Source/Runtime/MediaAssets/Public/MediaPlayer.h
--- a/Engine/Source/Runtime/MediaAssets/Public/MediaPlayer.h
+++ b/Engine/Source/Runtime/MediaAssets/Public/MediaPlayer.h
@@ -5,6 +5,7 @@
 #include "CoreTypes.h"
 #include "Containers/UnrealString.h"
 #include "Delegates/Delegate.h"
+#include "IMediaOptions.h"
 #include "Math/Quat.h"
 #include "Math/Rotator.h"
 #include "Templates/SharedPointer.h"
@@ -731,6 +732,14 @@ public:
 	UFUNCTION(BlueprintCallable, Category="Media|MediaPlayer")
 	bool SetLooping(bool Looping);
 
+	/**
+	 * Sets the media options used by the player.
+	 *
+	 * @param Options Options to pass to the player.
+	 */
+	UFUNCTION(BlueprintCallable, Category = "Media|MediaPlayer")
+	void SetMediaOptions(const UMediaSource* Options);
+
 	/**
 	 * Changes the media's playback rate.
 	 *
diff --git a/Engine/Source/Runtime/MediaAssets/Public/MediaSource.h b/Engine/Source/Runtime/MediaAssets/Public/MediaSource.h
--- a/Engine/Source/Runtime/MediaAssets/Public/MediaSource.h
+++ b/Engine/Source/Runtime/MediaAssets/Public/MediaSource.h
@@ -4,6 +4,7 @@
 
 #include "CoreTypes.h"
 #include "IMediaOptions.h"
+#include "Misc/Variant.h"
 #include "UObject/ObjectMacros.h"
 #include "UObject/Object.h"
 #include "IMediaOptions.h"
@@ -58,4 +59,34 @@ public:
 	virtual FString GetMediaOption(const FName& Key, const FString& DefaultValue) const override;
 	virtual FText GetMediaOption(const FName& Key, const FText& DefaultValue) const override;
 	virtual bool HasMediaOption(const FName& Key) const override;
+
+	/** Set a boolean parameter to pass to the player. */
+	UFUNCTION(BlueprintCallable, meta=(DisplayName = "SetMediaOption (boolean)"), Category = "Media|MediaSource")
+	void SetMediaOptionBool(const FName& Key, bool Value);
+	/** Set a float parameter to pass to the player. */
+	UFUNCTION(BlueprintCallable, meta = (DisplayName = "SetMediaOption (float)"), Category = "Media|MediaSource")
+	void SetMediaOptionFloat(const FName& Key, float Value);
+	/** Set a double parameter to pass to the player. */
+	void SetMediaOptionDouble(const FName& Key, double Value);
+	/** Set an integer64 parameter to pass to the player. */
+	UFUNCTION(BlueprintCallable, meta = (DisplayName = "SetMediaOption (integer64)"), Category = "Media|MediaSource")
+	void SetMediaOptionInt64(const FName& Key, int64 Value);
+	/** Set a string parameter to pass to the player. */
+	UFUNCTION(BlueprintCallable, meta = (DisplayName = "SetMediaOption (string)"), Category = "Media|MediaSource")
+	void SetMediaOptionString(const FName& Key, const FString& Value);
+
+private:
+	/** Holds our media options. */
+	TMap<FName, FVariant> MediaOptionsMap;
+
+	/**
+	 * Get the media option specified by the Key as a Variant.
+	 * Returns nullptr if the Key does not exist.
+	 */
+	const FVariant* GetMediaOption(const FName& Key) const;
+
+	/**
+	 * Sets the media option specified by Key to the supplied Variant.
+	 */
+	void SetMediaOption(const FName& Key, FVariant& Value);
 };
diff --git a/Engine/Source/Runtime/MediaUtils/Private/MediaPlayerFacade.cpp b/Engine/Source/Runtime/MediaUtils/Private/MediaPlayerFacade.cpp
--- a/Engine/Source/Runtime/MediaUtils/Private/MediaPlayerFacade.cpp
+++ b/Engine/Source/Runtime/MediaUtils/Private/MediaPlayerFacade.cpp
@@ -566,6 +566,11 @@ bool FMediaPlayerFacade::SetLooping(bool Looping)
 }
 
 
+void FMediaPlayerFacade::SetMediaOptions(const IMediaOptions* Options)
+{
+}
+
+
 bool FMediaPlayerFacade::SetRate(float Rate)
 {
 	if (!Player.IsValid() || !Player->GetControls().SetRate(Rate))
diff --git a/Engine/Source/Runtime/MediaUtils/Public/MediaPlayerFacade.h b/Engine/Source/Runtime/MediaUtils/Public/MediaPlayerFacade.h
--- a/Engine/Source/Runtime/MediaUtils/Public/MediaPlayerFacade.h
+++ b/Engine/Source/Runtime/MediaUtils/Public/MediaPlayerFacade.h
@@ -567,6 +567,11 @@ public:
 	 */
 	bool SetLooping(bool Looping);
 
+	/**
+	 * Changes media ooptions on the player.
+	 */
+	void SetMediaOptions(const IMediaOptions* Options);
+
 	/**
 	 * Changes the media's playback rate.
 	 *
