MediaFramework: Compile fix for mac.

diff --git a/Engine/Source/Runtime/MediaAssets/Private/Assets/MediaSource.cpp b/Engine/Source/Runtime/MediaAssets/Private/Assets/MediaSource.cpp
--- a/Engine/Source/Runtime/MediaAssets/Private/Assets/MediaSource.cpp
+++ b/Engine/Source/Runtime/MediaAssets/Private/Assets/MediaSource.cpp
@@ -13,7 +13,7 @@ FName UMediaSource::GetDesiredPlayerName() const
 
 bool UMediaSource::GetMediaOption(const FName& Key, bool DefaultValue) const
 {
-	const FVariant* Variant = GetMediaOption(Key);
+	const FVariant* Variant = GetMediaOptionDefault(Key);
 	if ((Variant != nullptr) && (Variant->GetType() == EVariantTypes::Bool))
 	{
 		return Variant->GetValue<bool>();
@@ -27,7 +27,7 @@ bool UMediaSource::GetMediaOption(const FName& Key, bool DefaultValue) const
 
 double UMediaSource::GetMediaOption(const FName& Key, double DefaultValue) const
 {
-	const FVariant* Variant = GetMediaOption(Key);
+	const FVariant* Variant = GetMediaOptionDefault(Key);
 	if ((Variant != nullptr) && (Variant->GetType() == EVariantTypes::Double))
 	{
 		return Variant->GetValue<double>();
@@ -41,7 +41,7 @@ double UMediaSource::GetMediaOption(const FName& Key, double DefaultValue) const
 
 int64 UMediaSource::GetMediaOption(const FName& Key, int64 DefaultValue) const
 {
-	const FVariant* Variant = GetMediaOption(Key);
+	const FVariant* Variant = GetMediaOptionDefault(Key);
 	if ((Variant != nullptr) && (Variant->GetType() == EVariantTypes::Int64))
 	{
 		return Variant->GetValue<int64>();
@@ -55,7 +55,7 @@ int64 UMediaSource::GetMediaOption(const FName& Key, int64 DefaultValue) const
 
 FString UMediaSource::GetMediaOption(const FName& Key, const FString& DefaultValue) const
 {
-	const FVariant* Variant = GetMediaOption(Key);
+	const FVariant* Variant = GetMediaOptionDefault(Key);
 	if ((Variant != nullptr) && (Variant->GetType() == EVariantTypes::String))
 	{
 		return Variant->GetValue<FString>();
@@ -113,7 +113,7 @@ void UMediaSource::SetMediaOptionString(const FName& Key, const FString& Value)
 }
 
 
-const FVariant* UMediaSource::GetMediaOption(const FName& Key) const
+const FVariant* UMediaSource::GetMediaOptionDefault(const FName& Key) const
 {
 	return MediaOptionsMap.Find(Key);
 }
diff --git a/Engine/Source/Runtime/MediaAssets/Public/MediaSource.h b/Engine/Source/Runtime/MediaAssets/Public/MediaSource.h
--- a/Engine/Source/Runtime/MediaAssets/Public/MediaSource.h
+++ b/Engine/Source/Runtime/MediaAssets/Public/MediaSource.h
@@ -83,7 +83,7 @@ private:
 	 * Get the media option specified by the Key as a Variant.
 	 * Returns nullptr if the Key does not exist.
 	 */
-	const FVariant* GetMediaOption(const FName& Key) const;
+	const FVariant* GetMediaOptionDefault(const FName& Key) const;
 
 	/**
 	 * Sets the media option specified by Key to the supplied Variant.
