Create a 10bit media texture for CharBGR10A2 media format.

diff --git a/Engine/Source/Runtime/MediaAssets/Private/Misc/MediaTextureResource.cpp b/Engine/Source/Runtime/MediaAssets/Private/Misc/MediaTextureResource.cpp
--- a/Engine/Source/Runtime/MediaAssets/Private/Misc/MediaTextureResource.cpp
+++ b/Engine/Source/Runtime/MediaAssets/Private/Misc/MediaTextureResource.cpp
@@ -73,6 +73,19 @@ namespace MediaTextureResource
 	}
 
 
+	EPixelFormat GetConvertedPixelFormat(const TSharedPtr<IMediaTextureSample, ESPMode::ThreadSafe>& Sample)
+	{
+		switch (Sample->GetFormat())
+		{
+		case EMediaTextureSampleFormat::CharBGR10A2:
+			return PF_A2B10G10R10;
+
+		default:
+			return PF_B8G8R8A8;
+		}
+	}
+
+
 	/**
 	 * Check whether the given sample requires a conversion shader.
 	 *
@@ -368,8 +381,9 @@ void FMediaTextureResource::ClearTexture(const FLinearColor& ClearColor, bool Sr
 {
 	// create output render target if we don't have one yet
 	const uint32 OutputCreateFlags = TexCreate_Dynamic | (SrgbOutput ? TexCreate_SRGB : 0);
+	const EPixelFormat OutputPixelFormat = PF_B8G8R8A8;
 
-	if ((ClearColor != CurrentClearColor) || !OutputTarget.IsValid() || ((OutputTarget->GetFlags() & OutputCreateFlags) != OutputCreateFlags))
+	if ((ClearColor != CurrentClearColor) || !OutputTarget.IsValid() || (OutputTarget->GetFormat() != OutputPixelFormat) || ((OutputTarget->GetFlags() & OutputCreateFlags) != OutputCreateFlags))
 	{
 		FRHIResourceCreateInfo CreateInfo = {
 			FClearValueBinding(ClearColor)
@@ -380,7 +394,7 @@ void FMediaTextureResource::ClearTexture(const FLinearColor& ClearColor, bool Sr
 		RHICreateTargetableShaderResource2D(
 			2,
 			2,
-			PF_B8G8R8A8,
+			OutputPixelFormat,
 			1,
 			OutputCreateFlags,
 			TexCreate_RenderTargetable,
@@ -415,7 +429,7 @@ void FMediaTextureResource::ClearTexture(const FLinearColor& ClearColor, bool Sr
 
 void FMediaTextureResource::ConvertSample(const TSharedPtr<IMediaTextureSample, ESPMode::ThreadSafe>& Sample, const FLinearColor& ClearColor, bool SrgbOutput)
 {
-	const EPixelFormat PixelFormat = MediaTextureResource::GetPixelFormat(Sample);
+	const EPixelFormat InputPixelFormat = MediaTextureResource::GetPixelFormat(Sample);
 
 	// get input texture
 	FRHITexture2D* InputTexture = nullptr;
@@ -441,7 +455,7 @@ void FMediaTextureResource::ConvertSample(const TSharedPtr<IMediaTextureSample,
 			const FIntPoint SampleDim = Sample->GetDim();
 
 			// create a new input render target if necessary
-			if (!InputTarget.IsValid() || (InputTarget->GetSizeXY() != SampleDim) || (InputTarget->GetFormat() != PixelFormat) || ((InputTarget->GetFlags() & InputCreateFlags) != InputCreateFlags))
+			if (!InputTarget.IsValid() || (InputTarget->GetSizeXY() != SampleDim) || (InputTarget->GetFormat() != InputPixelFormat) || ((InputTarget->GetFlags() & InputCreateFlags) != InputCreateFlags))
 			{
 				TRefCountPtr<FRHITexture2D> DummyTexture2DRHI;
 				FRHIResourceCreateInfo CreateInfo;
@@ -449,7 +463,7 @@ void FMediaTextureResource::ConvertSample(const TSharedPtr<IMediaTextureSample,
 				RHICreateTargetableShaderResource2D(
 					SampleDim.X,
 					SampleDim.Y,
-					PixelFormat,
+					InputPixelFormat,
 					1,
 					InputCreateFlags,
 					TexCreate_RenderTargetable,
@@ -473,8 +487,9 @@ void FMediaTextureResource::ConvertSample(const TSharedPtr<IMediaTextureSample,
 	// create output render target if necessary
 	const uint32 OutputCreateFlags = TexCreate_Dynamic | (SrgbOutput ? TexCreate_SRGB : 0);
 	const FIntPoint OutputDim = Sample->GetOutputDim();
+	const EPixelFormat OutputPixelFormat = MediaTextureResource::GetConvertedPixelFormat(Sample);
 
-	if ((ClearColor != CurrentClearColor) || !OutputTarget.IsValid() || (OutputTarget->GetSizeXY() != OutputDim) || (OutputTarget->GetFormat() != PF_B8G8R8A8) || ((OutputTarget->GetFlags() & OutputCreateFlags) != OutputCreateFlags))
+	if ((ClearColor != CurrentClearColor) || !OutputTarget.IsValid() || (OutputTarget->GetSizeXY() != OutputDim) || (OutputTarget->GetFormat() != OutputPixelFormat) || ((OutputTarget->GetFlags() & OutputCreateFlags) != OutputCreateFlags))
 	{
 		TRefCountPtr<FRHITexture2D> DummyTexture2DRHI;
 		
@@ -485,7 +500,7 @@ void FMediaTextureResource::ConvertSample(const TSharedPtr<IMediaTextureSample,
 		RHICreateTargetableShaderResource2D(
 			OutputDim.X,
 			OutputDim.Y,
-			PF_B8G8R8A8,
+			OutputPixelFormat,
 			1,
 			OutputCreateFlags,
 			TexCreate_RenderTargetable,
