Hide 'Source File' property of media textures in editor.

diff --git a/Engine/Source/Runtime/MediaAssets/Public/MediaTexture.h b/Engine/Source/Runtime/MediaAssets/Public/MediaTexture.h
--- a/Engine/Source/Runtime/MediaAssets/Public/MediaTexture.h
+++ b/Engine/Source/Runtime/MediaAssets/Public/MediaTexture.h
@@ -24,7 +24,7 @@ class UMediaPlayer;
 /**
  * Implements a texture asset for rendering video tracks from UMediaPlayer assets.
  */
-UCLASS(hidecategories=(Adjustments, Compositing, LevelOfDetail, Object))
+UCLASS(hidecategories=(Adjustments, Compositing, LevelOfDetail, ImportSettings, Object))
 class MEDIAASSETS_API UMediaTexture
 	: public UTexture
 {
