Fix crash after double clicking duplicated media player in content browser

diff --git a/Engine/Source/Runtime/MediaAssets/Private/Assets/MediaPlayer.cpp b/Engine/Source/Runtime/MediaAssets/Private/Assets/MediaPlayer.cpp
--- a/Engine/Source/Runtime/MediaAssets/Private/Assets/MediaPlayer.cpp
+++ b/Engine/Source/Runtime/MediaAssets/Private/Assets/MediaPlayer.cpp
@@ -724,6 +724,11 @@ void UMediaPlayer::PostDuplicate(bool bDuplicateForPIE)
 	{
 		PlayerGuid = FGuid::NewGuid();
 		PlayerFacade->SetGuid(PlayerGuid);
+
+		if (!Playlist)
+		{
+			Playlist = NewObject<UMediaPlaylist>(GetTransientPackage(), NAME_None, RF_Transactional | RF_Transient);
+		}
 	}
 }
 
