Renderpass updates.
Movie and media.

diff --git a/Engine/Plugins/Experimental/PixelStreaming/Source/PixelStreaming/Private/NvVideoEncoder.cpp b/Engine/Plugins/Experimental/PixelStreaming/Source/PixelStreaming/Private/NvVideoEncoder.cpp
--- a/Engine/Plugins/Experimental/PixelStreaming/Source/PixelStreaming/Private/NvVideoEncoder.cpp
+++ b/Engine/Plugins/Experimental/PixelStreaming/Source/PixelStreaming/Private/NvVideoEncoder.cpp
@@ -427,42 +427,47 @@ void FNvVideoEncoder::FNvVideoEncoderImpl::CopyBackBuffer(const FTexture2DRHIRef
 	}
 	else // Texture format mismatch, use a shader to do the copy.
 	{
-		SetRenderTarget(RHICmdList, ResolvedBackBuffer, FTextureRHIRef());
-		RHICmdList.SetViewport(0, 0, 0.0f, ResolvedBackBuffer->GetSizeX(), ResolvedBackBuffer->GetSizeY(), 1.0f);
-	
-		FGraphicsPipelineStateInitializer GraphicsPSOInit;
-		RHICmdList.ApplyCachedRenderTargets(GraphicsPSOInit);
-		GraphicsPSOInit.BlendState = TStaticBlendState<>::GetRHI();
-		GraphicsPSOInit.RasterizerState = TStaticRasterizerState<>::GetRHI();
-		GraphicsPSOInit.DepthStencilState = TStaticDepthStencilState<false, CF_Always>::GetRHI();
-	
-		TShaderMap<FGlobalShaderType>* ShaderMap = GetGlobalShaderMap(GMaxRHIFeatureLevel);
-		TShaderMapRef<FScreenVS> VertexShader(ShaderMap);
-		TShaderMapRef<FScreenPS> PixelShader(ShaderMap);
-	
-		GraphicsPSOInit.BoundShaderState.VertexDeclarationRHI = RendererModule->GetFilterVertexDeclaration().VertexDeclarationRHI;
-		GraphicsPSOInit.BoundShaderState.VertexShaderRHI = GETSAFERHISHADER_VERTEX(*VertexShader);
-		GraphicsPSOInit.BoundShaderState.PixelShaderRHI = GETSAFERHISHADER_PIXEL(*PixelShader);
-		GraphicsPSOInit.PrimitiveType = PT_TriangleList;
-	
-		SetGraphicsPipelineState(RHICmdList, GraphicsPSOInit);
-	
-		if (ResolvedBackBuffer->GetSizeX() != BackBuffer->GetSizeX() || ResolvedBackBuffer->GetSizeY() != BackBuffer->GetSizeY())
-			PixelShader->SetParameters(RHICmdList, TStaticSamplerState<SF_Bilinear>::GetRHI(), BackBuffer);
-		else
-			PixelShader->SetParameters(RHICmdList, TStaticSamplerState<SF_Point>::GetRHI(), BackBuffer);
-	
-		RendererModule->DrawRectangle(
-			RHICmdList,
-			0, 0,									// Dest X, Y
-			ResolvedBackBuffer->GetSizeX(),			// Dest Width
-			ResolvedBackBuffer->GetSizeY(),			// Dest Height
-			0, 0,									// Source U, V
-			1, 1,									// Source USize, VSize
-			ResolvedBackBuffer->GetSizeXY(),		// Target buffer size
-			FIntPoint(1, 1),						// Source texture size
-			*VertexShader,
-			EDRF_Default);
+		// #todo-renderpasses there's no explicit resolve here? Do we need one?
+		FRHIRenderPassInfo RPInfo(ResolvedBackBuffer, ERenderTargetActions::Load_Store);
+		RHICmdList.BeginRenderPass(RPInfo, TEXT("CopyBackbuffer"));
+		{
+			RHICmdList.SetViewport(0, 0, 0.0f, ResolvedBackBuffer->GetSizeX(), ResolvedBackBuffer->GetSizeY(), 1.0f);
+
+			FGraphicsPipelineStateInitializer GraphicsPSOInit;
+			RHICmdList.ApplyCachedRenderTargets(GraphicsPSOInit);
+			GraphicsPSOInit.BlendState = TStaticBlendState<>::GetRHI();
+			GraphicsPSOInit.RasterizerState = TStaticRasterizerState<>::GetRHI();
+			GraphicsPSOInit.DepthStencilState = TStaticDepthStencilState<false, CF_Always>::GetRHI();
+
+			TShaderMap<FGlobalShaderType>* ShaderMap = GetGlobalShaderMap(GMaxRHIFeatureLevel);
+			TShaderMapRef<FScreenVS> VertexShader(ShaderMap);
+			TShaderMapRef<FScreenPS> PixelShader(ShaderMap);
+
+			GraphicsPSOInit.BoundShaderState.VertexDeclarationRHI = RendererModule->GetFilterVertexDeclaration().VertexDeclarationRHI;
+			GraphicsPSOInit.BoundShaderState.VertexShaderRHI = GETSAFERHISHADER_VERTEX(*VertexShader);
+			GraphicsPSOInit.BoundShaderState.PixelShaderRHI = GETSAFERHISHADER_PIXEL(*PixelShader);
+			GraphicsPSOInit.PrimitiveType = PT_TriangleList;
+
+			SetGraphicsPipelineState(RHICmdList, GraphicsPSOInit);
+
+			if (ResolvedBackBuffer->GetSizeX() != BackBuffer->GetSizeX() || ResolvedBackBuffer->GetSizeY() != BackBuffer->GetSizeY())
+				PixelShader->SetParameters(RHICmdList, TStaticSamplerState<SF_Bilinear>::GetRHI(), BackBuffer);
+			else
+				PixelShader->SetParameters(RHICmdList, TStaticSamplerState<SF_Point>::GetRHI(), BackBuffer);
+
+			RendererModule->DrawRectangle(
+				RHICmdList,
+				0, 0,									// Dest X, Y
+				ResolvedBackBuffer->GetSizeX(),			// Dest Width
+				ResolvedBackBuffer->GetSizeY(),			// Dest Height
+				0, 0,									// Source U, V
+				1, 1,									// Source USize, VSize
+				ResolvedBackBuffer->GetSizeXY(),		// Target buffer size
+				FIntPoint(1, 1),						// Source texture size
+				*VertexShader,
+				EDRF_Default);
+		}
+		RHICmdList.EndRenderPass();
 	}
 }
 
diff --git a/Engine/Plugins/Media/AvfMedia/Source/AvfMedia/Private/Player/AvfMediaVideoSampler.cpp b/Engine/Plugins/Media/AvfMedia/Source/AvfMedia/Private/Player/AvfMediaVideoSampler.cpp
--- a/Engine/Plugins/Media/AvfMedia/Source/AvfMedia/Private/Player/AvfMediaVideoSampler.cpp
+++ b/Engine/Plugins/Media/AvfMedia/Source/AvfMedia/Private/Player/AvfMediaVideoSampler.cpp
@@ -270,29 +270,31 @@ void FAvfMediaVideoSampler::Tick()
 				TShaderMapRef<FMediaShadersVS> VertexShader(ShaderMap);
 				TShaderMapRef<FYCbCrConvertPS> PixelShader(ShaderMap);
 				
-				SetRenderTarget(RHICmdList, ShaderResource, nullptr, ESimpleRenderTargetMode::EExistingColorAndDepth, FExclusiveDepthStencil::DepthNop_StencilNop);
-				
-				FGraphicsPipelineStateInitializer GraphicsPSOInit;
-				RHICmdList.ApplyCachedRenderTargets(GraphicsPSOInit);
-				
-				GraphicsPSOInit.BlendState = TStaticBlendStateWriteMask<CW_RGBA, CW_NONE, CW_NONE, CW_NONE, CW_NONE, CW_NONE, CW_NONE, CW_NONE>::GetRHI();
-				GraphicsPSOInit.RasterizerState = TStaticRasterizerState<>::GetRHI();
-				GraphicsPSOInit.DepthStencilState = TStaticDepthStencilState<false, CF_Always>::GetRHI();
-				
-				GraphicsPSOInit.BoundShaderState.VertexDeclarationRHI = GMediaVertexDeclaration.VertexDeclarationRHI;
-				GraphicsPSOInit.BoundShaderState.VertexShaderRHI = GETSAFERHISHADER_VERTEX(*VertexShader);
-				GraphicsPSOInit.BoundShaderState.PixelShaderRHI = GETSAFERHISHADER_PIXEL(*PixelShader);
-				GraphicsPSOInit.PrimitiveType = PT_TriangleStrip;
-				
-				SetGraphicsPipelineState(RHICmdList, GraphicsPSOInit);
-				
-				PixelShader->SetParameters(RHICmdList, YTex, UVTex, MediaShaders::YuvToSrgbPs4, true);
-				
-				RHICmdList.SetStreamSource(0, CreateTempMediaVertexBuffer(), 0);
-				RHICmdList.SetViewport(0, 0, 0.0f, YWidth, YHeight, 1.0f);
+				FRHIRenderPassInfo RPInfo(ShaderResource, ERenderTargetActions::Load_Store);
+				RHICmdList.BeginRenderPass(RPInfo, TEXT("AvfMediaSampler"));
+				{
+					FGraphicsPipelineStateInitializer GraphicsPSOInit;
+					RHICmdList.ApplyCachedRenderTargets(GraphicsPSOInit);
 
-				RHICmdList.DrawPrimitive(0, 2, 1);
-				
+					GraphicsPSOInit.BlendState = TStaticBlendStateWriteMask<CW_RGBA, CW_NONE, CW_NONE, CW_NONE, CW_NONE, CW_NONE, CW_NONE, CW_NONE>::GetRHI();
+					GraphicsPSOInit.RasterizerState = TStaticRasterizerState<>::GetRHI();
+					GraphicsPSOInit.DepthStencilState = TStaticDepthStencilState<false, CF_Always>::GetRHI();
+
+					GraphicsPSOInit.BoundShaderState.VertexDeclarationRHI = GMediaVertexDeclaration.VertexDeclarationRHI;
+					GraphicsPSOInit.BoundShaderState.VertexShaderRHI = GETSAFERHISHADER_VERTEX(*VertexShader);
+					GraphicsPSOInit.BoundShaderState.PixelShaderRHI = GETSAFERHISHADER_PIXEL(*PixelShader);
+					GraphicsPSOInit.PrimitiveType = PT_TriangleStrip;
+
+					SetGraphicsPipelineState(RHICmdList, GraphicsPSOInit);
+
+					PixelShader->SetParameters(RHICmdList, YTex, UVTex, MediaShaders::YuvToSrgbPs4, true);
+
+					RHICmdList.SetStreamSource(0, CreateTempMediaVertexBuffer(), 0);
+					RHICmdList.SetViewport(0, 0, 0.0f, YWidth, YHeight, 1.0f);
+
+					RHICmdList.DrawPrimitive(0, 2, 1);
+				}
+				RHICmdList.EndRenderPass();
 				RHICmdList.CopyToResolveTarget(ShaderResource, ShaderResource, FResolveParams());
 			}
 			
diff --git a/Engine/Plugins/Media/WebMMedia/Source/WebMMedia/Private/Player/WebMVideoDecoder.cpp b/Engine/Plugins/Media/WebMMedia/Source/WebMMedia/Private/Player/WebMVideoDecoder.cpp
--- a/Engine/Plugins/Media/WebMMedia/Source/WebMMedia/Private/Player/WebMVideoDecoder.cpp
+++ b/Engine/Plugins/Media/WebMMedia/Source/WebMMedia/Private/Player/WebMVideoDecoder.cpp
@@ -233,33 +233,36 @@ void FWebMVideoDecoder::ConvertYUVToRGBAndSubmit(const FConvertParams& Params)
 		}
 
 		FRHITexture* RenderTarget = VideoSample->GetTexture();
-		SetRenderTargets(CommandList, 1, &RenderTarget, nullptr, ESimpleRenderTargetMode::EExistingColorAndDepth, FExclusiveDepthStencil::DepthNop_StencilNop);
-
-		// configure media shaders
-		auto ShaderMap = GetGlobalShaderMap(GMaxRHIFeatureLevel);
+		FRHIRenderPassInfo RPInfo(RenderTarget, ERenderTargetActions::Load_Store);
+		CommandList.BeginRenderPass(RPInfo, TEXT("ConvertYUVtoRGBA"));
+		{
+			// configure media shaders
+			auto ShaderMap = GetGlobalShaderMap(GMaxRHIFeatureLevel);
 
-		TShaderMapRef<FYUVConvertPS> PixelShader(ShaderMap);
-		TShaderMapRef<FMediaShadersVS> VertexShader(ShaderMap);
+			TShaderMapRef<FYUVConvertPS> PixelShader(ShaderMap);
+			TShaderMapRef<FMediaShadersVS> VertexShader(ShaderMap);
 
-		FGraphicsPipelineStateInitializer GraphicsPSOInit;
-		{
-			CommandList.ApplyCachedRenderTargets(GraphicsPSOInit);
-			GraphicsPSOInit.BlendState = TStaticBlendStateWriteMask<CW_RGBA, CW_NONE, CW_NONE, CW_NONE, CW_NONE, CW_NONE, CW_NONE, CW_NONE>::GetRHI();
-			GraphicsPSOInit.RasterizerState = TStaticRasterizerState<>::GetRHI();
-			GraphicsPSOInit.DepthStencilState = TStaticDepthStencilState<false, CF_Always>::GetRHI();
-			GraphicsPSOInit.BoundShaderState.VertexDeclarationRHI = GMoviePlayerResources.VertexDeclarationRHI;
-			GraphicsPSOInit.BoundShaderState.VertexShaderRHI = GETSAFERHISHADER_VERTEX(*VertexShader);
-			GraphicsPSOInit.BoundShaderState.PixelShaderRHI = GETSAFERHISHADER_PIXEL(*PixelShader);
-			GraphicsPSOInit.PrimitiveType = PT_TriangleStrip;
-		}
+			FGraphicsPipelineStateInitializer GraphicsPSOInit;
+			{
+				CommandList.ApplyCachedRenderTargets(GraphicsPSOInit);
+				GraphicsPSOInit.BlendState = TStaticBlendStateWriteMask<CW_RGBA, CW_NONE, CW_NONE, CW_NONE, CW_NONE, CW_NONE, CW_NONE, CW_NONE>::GetRHI();
+				GraphicsPSOInit.RasterizerState = TStaticRasterizerState<>::GetRHI();
+				GraphicsPSOInit.DepthStencilState = TStaticDepthStencilState<false, CF_Always>::GetRHI();
+				GraphicsPSOInit.BoundShaderState.VertexDeclarationRHI = GMoviePlayerResources.VertexDeclarationRHI;
+				GraphicsPSOInit.BoundShaderState.VertexShaderRHI = GETSAFERHISHADER_VERTEX(*VertexShader);
+				GraphicsPSOInit.BoundShaderState.PixelShaderRHI = GETSAFERHISHADER_PIXEL(*PixelShader);
+				GraphicsPSOInit.PrimitiveType = PT_TriangleStrip;
+			}
 
-		SetGraphicsPipelineState(CommandList, GraphicsPSOInit);
-		PixelShader->SetParameters(CommandList, DecodedY->GetTexture2D(), DecodedU->GetTexture2D(), DecodedV->GetTexture2D(), MediaShaders::YuvToSrgbDefault, true);
+			SetGraphicsPipelineState(CommandList, GraphicsPSOInit);
+			PixelShader->SetParameters(CommandList, DecodedY->GetTexture2D(), DecodedU->GetTexture2D(), DecodedV->GetTexture2D(), MediaShaders::YuvToSrgbDefault, true);
 
-		// draw full-size quad
-		CommandList.SetViewport(0, 0, 0.0f, Image->w, Image->d_h, 1.0f);
-		CommandList.SetStreamSource(0, GMoviePlayerResources.VertexBufferRHI, 0);
-		CommandList.DrawPrimitive(0, 2, 1);
+			// draw full-size quad
+			CommandList.SetViewport(0, 0, 0.0f, Image->w, Image->d_h, 1.0f);
+			CommandList.SetStreamSource(0, GMoviePlayerResources.VertexBufferRHI, 0);
+			CommandList.DrawPrimitive(0, 2, 1);
+		}
+		CommandList.EndRenderPass();
 		CommandList.CopyToResolveTarget(RenderTarget, RenderTarget, FResolveParams());
 
 		Samples->AddVideo(VideoSample.ToSharedRef());
diff --git a/Engine/Source/Runtime/MediaAssets/Private/Misc/MediaTextureResource.cpp b/Engine/Source/Runtime/MediaAssets/Private/Misc/MediaTextureResource.cpp
--- a/Engine/Source/Runtime/MediaAssets/Private/Misc/MediaTextureResource.cpp
+++ b/Engine/Source/Runtime/MediaAssets/Private/Misc/MediaTextureResource.cpp
@@ -416,9 +416,9 @@ void FMediaTextureResource::ClearTexture(const FLinearColor& ClearColor, bool Sr
 	// draw the clear color
 	FRHICommandListImmediate& CommandList = FRHICommandListExecutor::GetImmediateCommandList();
 	{
-		FRHIRenderTargetView View = FRHIRenderTargetView(RenderTargetTextureRHI, ERenderTargetLoadAction::EClear);
-		FRHISetRenderTargetsInfo Info(1, &View, FRHIDepthRenderTargetView());
-		CommandList.SetRenderTargetsAndClear(Info);
+		FRHIRenderPassInfo RPInfo(RenderTargetTextureRHI, ERenderTargetActions::Clear_Store);
+		CommandList.BeginRenderPass(RPInfo, TEXT("ClearTexture"));
+		CommandList.EndRenderPass();
 		CommandList.SetViewport(0, 0, 0.0f, RenderTargetTextureRHI->GetSizeX(), RenderTargetTextureRHI->GetSizeY(), 1.0f);
 		CommandList.TransitionResource(EResourceTransitionAccess::EReadable, RenderTargetTextureRHI);
 	}
@@ -526,118 +526,122 @@ void FMediaTextureResource::ConvertSample(const TSharedPtr<IMediaTextureSample,
 
 		FGraphicsPipelineStateInitializer GraphicsPSOInit;
 		FRHITexture* RenderTarget = RenderTargetTextureRHI.GetReference();
-		SetRenderTargets(CommandList, 1, &RenderTarget, nullptr, ESimpleRenderTargetMode::EExistingColorAndDepth, FExclusiveDepthStencil::DepthNop_StencilNop);
 
-		CommandList.ApplyCachedRenderTargets(GraphicsPSOInit);
-		CommandList.SetViewport(0, 0, 0.0f, OutputDim.X, OutputDim.Y, 1.0f);
+		FRHIRenderPassInfo RPInfo(RenderTarget, ERenderTargetActions::Load_Store);
+		CommandList.BeginRenderPass(RPInfo, TEXT("ConvertMedia"));
+		{
+			CommandList.ApplyCachedRenderTargets(GraphicsPSOInit);
+			CommandList.SetViewport(0, 0, 0.0f, OutputDim.X, OutputDim.Y, 1.0f);
 
-		GraphicsPSOInit.DepthStencilState = TStaticDepthStencilState<false, CF_Always>::GetRHI();
-		GraphicsPSOInit.RasterizerState = TStaticRasterizerState<>::GetRHI();
-		GraphicsPSOInit.BlendState = TStaticBlendStateWriteMask<CW_RGBA, CW_NONE, CW_NONE, CW_NONE, CW_NONE, CW_NONE, CW_NONE, CW_NONE>::GetRHI();
-		GraphicsPSOInit.PrimitiveType = PT_TriangleStrip;
+			GraphicsPSOInit.DepthStencilState = TStaticDepthStencilState<false, CF_Always>::GetRHI();
+			GraphicsPSOInit.RasterizerState = TStaticRasterizerState<>::GetRHI();
+			GraphicsPSOInit.BlendState = TStaticBlendStateWriteMask<CW_RGBA, CW_NONE, CW_NONE, CW_NONE, CW_NONE, CW_NONE, CW_NONE, CW_NONE>::GetRHI();
+			GraphicsPSOInit.PrimitiveType = PT_TriangleStrip;
 
-		// configure media shaders
-		auto ShaderMap = GetGlobalShaderMap(GMaxRHIFeatureLevel);
-		TShaderMapRef<FMediaShadersVS> VertexShader(ShaderMap);
+			// configure media shaders
+			auto ShaderMap = GetGlobalShaderMap(GMaxRHIFeatureLevel);
+			TShaderMapRef<FMediaShadersVS> VertexShader(ShaderMap);
 
-		GraphicsPSOInit.BoundShaderState.VertexDeclarationRHI = GMediaVertexDeclaration.VertexDeclarationRHI;
-		GraphicsPSOInit.BoundShaderState.VertexShaderRHI = GETSAFERHISHADER_VERTEX(*VertexShader);
+			GraphicsPSOInit.BoundShaderState.VertexDeclarationRHI = GMediaVertexDeclaration.VertexDeclarationRHI;
+			GraphicsPSOInit.BoundShaderState.VertexShaderRHI = GETSAFERHISHADER_VERTEX(*VertexShader);
 
-		switch (Sample->GetFormat())
-		{
-		case EMediaTextureSampleFormat::CharAYUV:
-		{
-			TShaderMapRef<FAYUVConvertPS> ConvertShader(ShaderMap);
-			GraphicsPSOInit.BoundShaderState.PixelShaderRHI = GETSAFERHISHADER_PIXEL(*ConvertShader);
-			SetGraphicsPipelineState(CommandList, GraphicsPSOInit);
-			ConvertShader->SetParameters(CommandList, InputTexture, MediaShaders::YuvToSrgbDefault, Sample->IsOutputSrgb());
-		}
-		break;
+			switch (Sample->GetFormat())
+			{
+			case EMediaTextureSampleFormat::CharAYUV:
+			{
+				TShaderMapRef<FAYUVConvertPS> ConvertShader(ShaderMap);
+				GraphicsPSOInit.BoundShaderState.PixelShaderRHI = GETSAFERHISHADER_PIXEL(*ConvertShader);
+				SetGraphicsPipelineState(CommandList, GraphicsPSOInit);
+				ConvertShader->SetParameters(CommandList, InputTexture, MediaShaders::YuvToSrgbDefault, Sample->IsOutputSrgb());
+			}
+			break;
 
-		case EMediaTextureSampleFormat::CharBMP:
-		{
-			TShaderMapRef<FBMPConvertPS> ConvertShader(ShaderMap);
-			GraphicsPSOInit.BoundShaderState.PixelShaderRHI = GETSAFERHISHADER_PIXEL(*ConvertShader);
-			SetGraphicsPipelineState(CommandList, GraphicsPSOInit);
-			ConvertShader->SetParameters(CommandList, InputTexture, OutputDim, Sample->IsOutputSrgb() && !SrgbOutput);
-		}
-		break;
+			case EMediaTextureSampleFormat::CharBMP:
+			{
+				TShaderMapRef<FBMPConvertPS> ConvertShader(ShaderMap);
+				GraphicsPSOInit.BoundShaderState.PixelShaderRHI = GETSAFERHISHADER_PIXEL(*ConvertShader);
+				SetGraphicsPipelineState(CommandList, GraphicsPSOInit);
+				ConvertShader->SetParameters(CommandList, InputTexture, OutputDim, Sample->IsOutputSrgb() && !SrgbOutput);
+			}
+			break;
 
-		case EMediaTextureSampleFormat::CharNV12:
-		{
-			TShaderMapRef<FNV12ConvertPS> ConvertShader(ShaderMap);
-			GraphicsPSOInit.BoundShaderState.PixelShaderRHI = GETSAFERHISHADER_PIXEL(*ConvertShader);
-			SetGraphicsPipelineState(CommandList, GraphicsPSOInit);
-			ConvertShader->SetParameters(CommandList, InputTexture, OutputDim, MediaShaders::YuvToSrgbDefault, Sample->IsOutputSrgb());
-		}
-		break;
+			case EMediaTextureSampleFormat::CharNV12:
+			{
+				TShaderMapRef<FNV12ConvertPS> ConvertShader(ShaderMap);
+				GraphicsPSOInit.BoundShaderState.PixelShaderRHI = GETSAFERHISHADER_PIXEL(*ConvertShader);
+				SetGraphicsPipelineState(CommandList, GraphicsPSOInit);
+				ConvertShader->SetParameters(CommandList, InputTexture, OutputDim, MediaShaders::YuvToSrgbDefault, Sample->IsOutputSrgb());
+			}
+			break;
 
-		case EMediaTextureSampleFormat::CharNV21:
-		{
-			TShaderMapRef<FNV21ConvertPS> ConvertShader(ShaderMap);
-			GraphicsPSOInit.BoundShaderState.PixelShaderRHI = GETSAFERHISHADER_PIXEL(*ConvertShader);
-			SetGraphicsPipelineState(CommandList, GraphicsPSOInit);
-			ConvertShader->SetParameters(CommandList, InputTexture, OutputDim, MediaShaders::YuvToSrgbDefault, Sample->IsOutputSrgb());
-		}
-		break;
+			case EMediaTextureSampleFormat::CharNV21:
+			{
+				TShaderMapRef<FNV21ConvertPS> ConvertShader(ShaderMap);
+				GraphicsPSOInit.BoundShaderState.PixelShaderRHI = GETSAFERHISHADER_PIXEL(*ConvertShader);
+				SetGraphicsPipelineState(CommandList, GraphicsPSOInit);
+				ConvertShader->SetParameters(CommandList, InputTexture, OutputDim, MediaShaders::YuvToSrgbDefault, Sample->IsOutputSrgb());
+			}
+			break;
 
-		case EMediaTextureSampleFormat::CharUYVY:
-		{
-			TShaderMapRef<FUYVYConvertPS> ConvertShader(ShaderMap);
-			GraphicsPSOInit.BoundShaderState.PixelShaderRHI = GETSAFERHISHADER_PIXEL(*ConvertShader);
-			SetGraphicsPipelineState(CommandList, GraphicsPSOInit);
-			ConvertShader->SetParameters(CommandList, InputTexture, MediaShaders::YuvToSrgbDefault, Sample->IsOutputSrgb());
-		}
-		break;
+			case EMediaTextureSampleFormat::CharUYVY:
+			{
+				TShaderMapRef<FUYVYConvertPS> ConvertShader(ShaderMap);
+				GraphicsPSOInit.BoundShaderState.PixelShaderRHI = GETSAFERHISHADER_PIXEL(*ConvertShader);
+				SetGraphicsPipelineState(CommandList, GraphicsPSOInit);
+				ConvertShader->SetParameters(CommandList, InputTexture, MediaShaders::YuvToSrgbDefault, Sample->IsOutputSrgb());
+			}
+			break;
 
-		case EMediaTextureSampleFormat::CharYUY2:
-		{
-			TShaderMapRef<FYUY2ConvertPS> ConvertShader(ShaderMap);
-			GraphicsPSOInit.BoundShaderState.PixelShaderRHI = GETSAFERHISHADER_PIXEL(*ConvertShader);
-			SetGraphicsPipelineState(CommandList, GraphicsPSOInit);
-			ConvertShader->SetParameters(CommandList, InputTexture, OutputDim, MediaShaders::YuvToSrgbDefault, Sample->IsOutputSrgb());
-		}
-		break;
+			case EMediaTextureSampleFormat::CharYUY2:
+			{
+				TShaderMapRef<FYUY2ConvertPS> ConvertShader(ShaderMap);
+				GraphicsPSOInit.BoundShaderState.PixelShaderRHI = GETSAFERHISHADER_PIXEL(*ConvertShader);
+				SetGraphicsPipelineState(CommandList, GraphicsPSOInit);
+				ConvertShader->SetParameters(CommandList, InputTexture, OutputDim, MediaShaders::YuvToSrgbDefault, Sample->IsOutputSrgb());
+			}
+			break;
 
-		case EMediaTextureSampleFormat::CharYVYU:
-		{
-			TShaderMapRef<FYVYUConvertPS> ConvertShader(ShaderMap);
-			GraphicsPSOInit.BoundShaderState.PixelShaderRHI = GETSAFERHISHADER_PIXEL(*ConvertShader);
-			SetGraphicsPipelineState(CommandList, GraphicsPSOInit);
-			ConvertShader->SetParameters(CommandList, InputTexture, MediaShaders::YuvToSrgbDefault, Sample->IsOutputSrgb());
-		}
-		break;
+			case EMediaTextureSampleFormat::CharYVYU:
+			{
+				TShaderMapRef<FYVYUConvertPS> ConvertShader(ShaderMap);
+				GraphicsPSOInit.BoundShaderState.PixelShaderRHI = GETSAFERHISHADER_PIXEL(*ConvertShader);
+				SetGraphicsPipelineState(CommandList, GraphicsPSOInit);
+				ConvertShader->SetParameters(CommandList, InputTexture, MediaShaders::YuvToSrgbDefault, Sample->IsOutputSrgb());
+			}
+			break;
 
-		case EMediaTextureSampleFormat::CharBGR10A2:
-		{
-			TShaderMapRef<FRGBConvertPS> ConvertShader(ShaderMap);
-			GraphicsPSOInit.BoundShaderState.PixelShaderRHI = GETSAFERHISHADER_PIXEL(*ConvertShader);
-			SetGraphicsPipelineState(CommandList, GraphicsPSOInit);
-			ConvertShader->SetParameters(CommandList, InputTexture, OutputDim, Sample->IsOutputSrgb());
-		}
-		break;
+			case EMediaTextureSampleFormat::CharBGR10A2:
+			{
+				TShaderMapRef<FRGBConvertPS> ConvertShader(ShaderMap);
+				GraphicsPSOInit.BoundShaderState.PixelShaderRHI = GETSAFERHISHADER_PIXEL(*ConvertShader);
+				SetGraphicsPipelineState(CommandList, GraphicsPSOInit);
+				ConvertShader->SetParameters(CommandList, InputTexture, OutputDim, Sample->IsOutputSrgb());
+			}
+			break;
 
-		case EMediaTextureSampleFormat::CharBGRA:
-		case EMediaTextureSampleFormat::FloatRGB:
-		case EMediaTextureSampleFormat::FloatRGBA:
-		{
-			TShaderMapRef<FRGBConvertPS> ConvertShader(ShaderMap);
-			GraphicsPSOInit.BoundShaderState.PixelShaderRHI = GETSAFERHISHADER_PIXEL(*ConvertShader);
-			SetGraphicsPipelineState(CommandList, GraphicsPSOInit);
-			ConvertShader->SetParameters(CommandList, InputTexture, OutputDim, false);
-		}
-		break;
+			case EMediaTextureSampleFormat::CharBGRA:
+			case EMediaTextureSampleFormat::FloatRGB:
+			case EMediaTextureSampleFormat::FloatRGBA:
+			{
+				TShaderMapRef<FRGBConvertPS> ConvertShader(ShaderMap);
+				GraphicsPSOInit.BoundShaderState.PixelShaderRHI = GETSAFERHISHADER_PIXEL(*ConvertShader);
+				SetGraphicsPipelineState(CommandList, GraphicsPSOInit);
+				ConvertShader->SetParameters(CommandList, InputTexture, OutputDim, false);
+			}
+			break;
 
-		default:
-			return; // unsupported format
-		}
+			default:
+				return; // unsupported format
+			}
 
-		// draw full size quad into render target
-		CommandList.SetStreamSource(0, CreateTempMediaVertexBuffer(), 0);
-		// set viewport to RT size
-		CommandList.SetViewport(0, 0, 0.0f, OutputDim.X, OutputDim.Y, 1.0f);
+			// draw full size quad into render target
+			CommandList.SetStreamSource(0, CreateTempMediaVertexBuffer(), 0);
+			// set viewport to RT size
+			CommandList.SetViewport(0, 0, 0.0f, OutputDim.X, OutputDim.Y, 1.0f);
 
-		CommandList.DrawPrimitive(0, 2, 1);
+			CommandList.DrawPrimitive(0, 2, 1);
+		}
+		CommandList.EndRenderPass();
 		CommandList.TransitionResource(EResourceTransitionAccess::EReadable, RenderTargetTextureRHI);
 	}
 
