Fixed file media source not being able to change paths.

diff --git a/Engine/Source/Runtime/MediaAssets/Private/Assets/FileMediaSource.cpp b/Engine/Source/Runtime/MediaAssets/Private/Assets/FileMediaSource.cpp
--- a/Engine/Source/Runtime/MediaAssets/Private/Assets/FileMediaSource.cpp
+++ b/Engine/Source/Runtime/MediaAssets/Private/Assets/FileMediaSource.cpp
@@ -44,6 +44,18 @@ void UFileMediaSource::SetFilePath(const FString& Path)
 	}
 }
 
+#if WITH_EDITOR
+void UFileMediaSource::PostEditChangeProperty(FPropertyChangedEvent& PropertyChangedEvent)
+{
+	Super::PostEditChangeProperty(PropertyChangedEvent);
+
+	if (PropertyChangedEvent.GetPropertyName() == GET_MEMBER_NAME_CHECKED(UFileMediaSource, FilePath))
+	{
+		ClearResolvedFullPath();
+	}
+}
+#endif
+
 void UFileMediaSource::ClearResolvedFullPath() const
 {
 	ResolvedFullPath.Empty();
diff --git a/Engine/Source/Runtime/MediaAssets/Public/FileMediaSource.h b/Engine/Source/Runtime/MediaAssets/Public/FileMediaSource.h
--- a/Engine/Source/Runtime/MediaAssets/Public/FileMediaSource.h
+++ b/Engine/Source/Runtime/MediaAssets/Public/FileMediaSource.h
@@ -65,6 +65,10 @@ public:
 	UFUNCTION(BlueprintCallable, Category="Media|FileMediaSource")
 	void SetFilePath(const FString& Path);
 
+#if WITH_EDITOR
+	virtual void PostEditChangeProperty(FPropertyChangedEvent& PropertyChangedEvent) override;
+#endif
+
 public:
 
 	//~ IMediaOptions interface
