Media Player blueprint function 'open source latent' added

diff --git a/Engine/Source/Runtime/MediaAssets/Private/Assets/MediaPlayer.cpp b/Engine/Source/Runtime/MediaAssets/Private/Assets/MediaPlayer.cpp
--- a/Engine/Source/Runtime/MediaAssets/Private/Assets/MediaPlayer.cpp
+++ b/Engine/Source/Runtime/MediaAssets/Private/Assets/MediaPlayer.cpp
@@ -3,6 +3,7 @@
 #include "MediaPlayer.h"
 #include "MediaAssetsPrivate.h"
 
+#include "Engine/Engine.h"
 #include "IMediaClock.h"
 #include "IMediaControls.h"
 #include "IMediaModule.h"
@@ -10,6 +11,7 @@
 #include "IMediaPlayerFactory.h"
 #include "IMediaTicker.h"
 #include "IMediaTracks.h"
+#include "LatentActions.h"
 #include "MediaPlayerFacade.h"
 #include "MediaPlayerOptions.h"
 #include "Misc/App.h"
@@ -815,6 +817,8 @@ void UMediaPlayer::HandlePlayerMediaEvent(EMediaEvent Event)
 {
 	MediaEvent.Broadcast(Event);
 
+	bool bPlayOnOpen = false;
+
 	switch(Event)
 	{
 	case EMediaEvent::MediaClosed:
@@ -823,17 +827,39 @@ void UMediaPlayer::HandlePlayerMediaEvent(EMediaEvent Event)
 
 	case EMediaEvent::MediaOpened:
 		PlayerFacade->SetCacheWindow(CacheAhead, FApp::IsGame() ? CacheBehindGame : CacheBehind);
-		PlayerFacade->SetLooping(Loop && (Playlist->Num() == 1));
+		if (PlayerFacade->ActivePlayerOptions.IsSet() && PlayerFacade->ActivePlayerOptions->Loop != EMediaPlayerOptionBooleanOverride::UseMediaPlayerSetting)
+		{
+			PlayerFacade->SetLooping(PlayerFacade->ActivePlayerOptions->Loop == EMediaPlayerOptionBooleanOverride::Enabled);
+		}
+		else
+		{
+			PlayerFacade->SetLooping(Loop && (Playlist->Num() == 1));
+		}
 		PlayerFacade->SetViewField(HorizontalFieldOfView, VerticalFieldOfView, true);
 		PlayerFacade->SetViewOrientation(FQuat(ViewRotation), true);
 		PlayerFacade->TimeDelay = TimeDelay;
 
 		OnMediaOpened.Broadcast(PlayerFacade->GetUrl());
 
-		if (PlayOnOpen || PlayOnNext)
+		if (PlayerFacade->ActivePlayerOptions.IsSet() && PlayerFacade->ActivePlayerOptions->PlayOnOpen != EMediaPlayerOptionBooleanOverride::UseMediaPlayerSetting)
+		{
+			bPlayOnOpen = (PlayerFacade->ActivePlayerOptions->PlayOnOpen == EMediaPlayerOptionBooleanOverride::Enabled);
+		}
+		else
+		{
+			bPlayOnOpen = (PlayOnOpen || PlayOnNext);
+		}
+
+		if (bPlayOnOpen)
 		{
 			PlayOnNext = false;
-			Play();
+			if (Play())
+			{
+				if (PlayerFacade->ActivePlayerOptions.IsSet() && !PlayerFacade->ActivePlayerOptions->SeekTime.IsZero() && SupportsSeeking())
+				{
+					Seek(PlayerFacade->ActivePlayerOptions->SeekTime);
+				}
+			}
 		}
 		break;
 
@@ -875,3 +901,102 @@ void UMediaPlayer::HandlePlayerMediaEvent(EMediaEvent Event)
 		break;
 	}
 }
+
+class FLatentOpenMediaSourceAction : public FPendingLatentAction
+{
+public:
+	FName ExecutionFunction;
+	int32 OutputLink;
+	FWeakObjectPtr CallbackTarget;
+
+	TWeakObjectPtr<UMediaPlayer> MediaPlayer;
+	float TimeRemaining;
+	bool& OutSuccess;
+	bool bSawError;
+	FString URL;
+
+	FLatentOpenMediaSourceAction(const FLatentActionInfo& LatentInfo, UMediaPlayer* InMediaPlayer, UMediaSource* InMediaSource, const FMediaPlayerOptions& InOptions, bool& InSuccess)
+		: ExecutionFunction(LatentInfo.ExecutionFunction)
+		, OutputLink(LatentInfo.Linkage)
+		, CallbackTarget(LatentInfo.CallbackTarget)
+		, MediaPlayer(InMediaPlayer)
+		, TimeRemaining(10.0)
+		, OutSuccess(InSuccess)
+		, bSawError(false)
+	{
+		if (InMediaSource)
+		{
+			URL = InMediaSource->GetUrl();
+		}
+
+		if (!InMediaPlayer->OpenSourceWithOptions(InMediaSource, InOptions))
+		{
+			UE_LOG(LogMediaAssets, Verbose, TEXT("Latent Media: After initial open started"));
+			bSawError = true;
+		}
+	}
+
+	virtual void UpdateOperation(FLatentResponse& Response) override
+	{
+		if (bSawError || !MediaPlayer.IsValid() || MediaPlayer.IsStale() || MediaPlayer->HasError())
+		{
+			UE_LOG(LogMediaAssets, Warning, TEXT("Latent Media: Deleted or Error"));
+			OutSuccess = false;
+			Response.FinishAndTriggerIf(false, ExecutionFunction, OutputLink, CallbackTarget);
+			return;
+		}
+
+		if (MediaPlayer->IsReady())
+		{
+			UE_LOG(LogMediaAssets, Verbose, TEXT("Latent Media: IsReady"));
+			OutSuccess = true;
+			Response.FinishAndTriggerIf(true, ExecutionFunction, OutputLink, CallbackTarget);
+			return;
+		}
+		else if (!MediaPlayer->IsPreparing())
+		{
+			UE_LOG(LogMediaAssets, Verbose, TEXT("Latent Media: Invalid State"));
+			OutSuccess = false;
+			Response.FinishAndTriggerIf(false, ExecutionFunction, OutputLink, CallbackTarget);
+			return;
+		}
+		else
+		{
+			UE_LOG(LogMediaAssets, Verbose, TEXT("Latent Media: Waiting for open ..."));
+		}
+
+		// Timed out
+		TimeRemaining -= Response.ElapsedTime();
+		if (TimeRemaining <= 0.0f)
+		{
+			UE_LOG(LogMediaAssets, Verbose, TEXT("Latent Media: Timed out."));
+			OutSuccess = false;
+			Response.FinishAndTriggerIf(true, ExecutionFunction, OutputLink, CallbackTarget);
+			return;
+		}
+	}
+
+#if WITH_EDITOR
+	// Returns a human readable description of the latent operation's current state
+	virtual FString GetDescription() const override
+	{
+		return FString::Printf(TEXT("Opening Media: %s"), *URL);
+	}
+#endif
+};
+
+void UMediaPlayer::OpenSourceLatent(const UObject* WorldContextObject, FLatentActionInfo LatentInfo, UMediaSource* MediaSource, const FMediaPlayerOptions& Options, bool& bSuccess)
+{
+	bSuccess = false;
+
+	if (UWorld* World = GEngine->GetWorldFromContextObject(WorldContextObject, EGetWorldErrorMode::LogAndReturnNull))
+	{
+		FLatentActionManager& LatentManager = World->GetLatentActionManager();
+		if (LatentManager.FindExistingAction<FLatentOpenMediaSourceAction>(LatentInfo.CallbackTarget, LatentInfo.UUID) == nullptr)
+		{
+			FLatentOpenMediaSourceAction* NewAction = new FLatentOpenMediaSourceAction(LatentInfo, this, MediaSource, Options, bSuccess);
+			LatentManager.AddNewAction(LatentInfo.CallbackTarget, LatentInfo.UUID, NewAction);
+		}
+	}
+}
+
diff --git a/Engine/Source/Runtime/MediaAssets/Public/MediaPlayer.h b/Engine/Source/Runtime/MediaAssets/Public/MediaPlayer.h
--- a/Engine/Source/Runtime/MediaAssets/Public/MediaPlayer.h
+++ b/Engine/Source/Runtime/MediaAssets/Public/MediaPlayer.h
@@ -12,6 +12,7 @@
 #include "UObject/Object.h"
 #include "UObject/ScriptMacros.h"
 #include "Misc/Guid.h"
+#include "Engine/LatentActionManager.h"
 #include "MediaPlayerOptions.h"
 
 #include "MediaPlayer.generated.h"
@@ -569,6 +570,19 @@ public:
 	UFUNCTION(BlueprintCallable, Category="Media|MediaPlayer")
 	bool OpenSourceWithOptions(UMediaSource* MediaSource, const FMediaPlayerOptions& Options);
 
+	/**
+	 * Open the specified media source with options using a latent action.
+	 *
+	 * A result of true indicates that the player opened the media source successfully.
+	 *
+	 * @param MediaSource The media source to open.
+	 * @param Options The media player options to apply.
+	 * @param bSuccess  Player opened the media.
+	 * @see Close, OpenFile, OpenPlaylist, OpenPlaylistIndex, OpenUrl, Reopen
+	 */
+	UFUNCTION(BlueprintCallable, Category = "Media|MediaPlayer", meta = (Latent, LatentInfo = "LatentInfo", WorldContext = "WorldContextObject"))
+	void OpenSourceLatent(const UObject* WorldContextObject, struct FLatentActionInfo LatentInfo, UMediaSource* MediaSource, const FMediaPlayerOptions& Options, bool& bSuccess);
+
 	/**
 	 * Opens the specified media URL.
 	 *
diff --git a/Engine/Source/Runtime/MediaUtils/Public/MediaPlayerOptions.h b/Engine/Source/Runtime/MediaUtils/Public/MediaPlayerOptions.h
--- a/Engine/Source/Runtime/MediaUtils/Public/MediaPlayerOptions.h
+++ b/Engine/Source/Runtime/MediaUtils/Public/MediaPlayerOptions.h
@@ -8,6 +8,14 @@
 
 #include "MediaPlayerOptions.generated.h"
 
+UENUM(BlueprintType)
+enum class EMediaPlayerOptionBooleanOverride : uint8
+{
+	UseMediaPlayerSetting,
+	Enabled,
+	Disabled
+};
+
 USTRUCT(BlueprintType)
 struct FMediaPlayerTrackOptions
 {
@@ -51,10 +59,22 @@ struct FMediaPlayerOptions
 {
 	GENERATED_BODY()
 
-	FMediaPlayerOptions()
+	FMediaPlayerOptions() :
+		SeekTime(0),
+		PlayOnOpen(EMediaPlayerOptionBooleanOverride::UseMediaPlayerSetting),
+		Loop(EMediaPlayerOptionBooleanOverride::UseMediaPlayerSetting)
 	{
 	}
 
 	UPROPERTY(BlueprintReadWrite, Category = "Tracks")
 	FMediaPlayerTrackOptions Tracks;
+
+	UPROPERTY(BlueprintReadWrite, Category = "Misc")
+	FTimespan SeekTime;
+
+	UPROPERTY(BlueprintReadWrite, Category = "Misc")
+	EMediaPlayerOptionBooleanOverride PlayOnOpen;
+
+	UPROPERTY(BlueprintReadWrite, Category = "Misc")
+	EMediaPlayerOptionBooleanOverride Loop;
 };
