Add media stat to track how many queued audio samples are waiting to play

diff --git a/Engine/Source/Runtime/MediaAssets/Private/Assets/MediaSoundComponent.cpp b/Engine/Source/Runtime/MediaAssets/Private/Assets/MediaSoundComponent.cpp
--- a/Engine/Source/Runtime/MediaAssets/Private/Assets/MediaSoundComponent.cpp
+++ b/Engine/Source/Runtime/MediaAssets/Private/Assets/MediaSoundComponent.cpp
@@ -31,6 +31,7 @@ FAutoConsoleVariableRef CVarSyncAudioAfterDropouts(
 
 DECLARE_FLOAT_COUNTER_STAT(TEXT("MediaUtils MediaSoundComponent Sync"), STAT_MediaUtils_MediaSoundComponentSync, STATGROUP_Media);
 DECLARE_FLOAT_COUNTER_STAT(TEXT("MediaUtils MediaSoundComponent SampleTime"), STAT_MediaUtils_MediaSoundComponentSampleTime, STATGROUP_Media);
+DECLARE_DWORD_COUNTER_STAT(TEXT("MediaUtils MediaSoundComponent Queued"), STAT_Media_SoundCompQueued, STATGROUP_Media);
 
 /* Static initialization
  *****************************************************************************/
@@ -444,6 +445,7 @@ int32 UMediaSoundComponent::OnGenerateAudio(float* OutAudio, int32 NumSamples)
 
 		SET_FLOAT_STAT(STAT_MediaUtils_MediaSoundComponentSync, FMath::Abs((Time - OutTime).GetTotalMilliseconds()));
 		SET_FLOAT_STAT(STAT_MediaUtils_MediaSoundComponentSampleTime, OutTime.GetTotalMilliseconds());
+		SET_DWORD_STAT(STAT_Media_SoundCompQueued, PinnedSampleQueue->Num());
 	}
 	else
 	{
