Implement envelope following to media sound component.

diff --git a/Engine/Source/Runtime/MediaAssets/Private/Assets/MediaSoundComponent.cpp b/Engine/Source/Runtime/MediaAssets/Private/Assets/MediaSoundComponent.cpp
--- a/Engine/Source/Runtime/MediaAssets/Private/Assets/MediaSoundComponent.cpp
+++ b/Engine/Source/Runtime/MediaAssets/Private/Assets/MediaSoundComponent.cpp
@@ -53,7 +53,12 @@ UMediaSoundComponent::UMediaSoundComponent(const FObjectInitializer& ObjectIniti
 	, RateAdjustment(1.0f)
 	, Resampler(new FMediaAudioResampler)
 	, FrameSyncOffset(0)
+	, EnvelopeFollowerAttackTime(10)
+	, EnvelopeFollowerReleaseTime(100)
+	, CurrentEnvelopeValue(0.0f)
 	, bSyncAudioAfterDropouts(false)
+	, bSpectralAnalysisEnabled(false)
+	, bEnvelopeFollowingEnabled(false)
 	, LastPlaySampleTime(FTimespan::MinValue())
 {
 	PrimaryComponentTick.bCanEverTick = true;
@@ -415,32 +420,57 @@ int32 UMediaSoundComponent::OnGenerateAudio(float* OutAudio, int32 NumSamples)
 
 		LastPlaySampleTime = OutTime;
 
-		if (bSpectralAnalysisEnabled)
+
+		if (bSpectralAnalysisEnabled || bEnvelopeFollowingEnabled)
 		{
-			// If we have stereo audio, sum to mono before sending to analyzer
+			float* BufferToUseForAnalysis = nullptr;
+			int32 NumFrames = NumSamples;
+			
 			if (NumChannels == 2)
 			{
-				int32 NumFrames = NumSamples / 2;
+				NumFrames = NumSamples / 2;
 
 				// Use the scratch buffer to sum the audio to mono
 				AudioScratchBuffer.Reset();
 				AudioScratchBuffer.AddUninitialized(NumFrames);
-				float* AudioScratchBufferPtr = AudioScratchBuffer.GetData();
+				BufferToUseForAnalysis = AudioScratchBuffer.GetData();
 				int32 SampleIndex = 0;
 				for (int32 FrameIndex = 0; FrameIndex < NumFrames; ++FrameIndex, SampleIndex += NumChannels)
 				{
-					AudioScratchBufferPtr[FrameIndex] = 0.5f * (OutAudio[SampleIndex] + OutAudio[SampleIndex + 1]);
+					BufferToUseForAnalysis[FrameIndex] = 0.5f * (OutAudio[SampleIndex] + OutAudio[SampleIndex + 1]);
 				}
-
-				SpectrumAnalyzer.PushAudio(AudioScratchBufferPtr, NumFrames);
 			}
 			else
 			{
-				SpectrumAnalyzer.PushAudio(OutAudio, NumSamples);
+				BufferToUseForAnalysis = OutAudio;
 			}
 
-			// Launch an analysis task with this audio
-			SpectrumAnalyzer.PerformAnalysisIfPossible(true, true);
+			if (bSpectralAnalysisEnabled)
+			{
+				SpectrumAnalyzer.PushAudio(BufferToUseForAnalysis, NumFrames);
+				SpectrumAnalyzer.PerformAnalysisIfPossible(true, true);
+			}
+
+			{
+				FScopeLock ScopeLock(&EnvelopeFollowerCriticalSection);
+				if (bEnvelopeFollowingEnabled)
+				{
+					if (bEnvelopeFollowerSettingsChanged)
+					{
+						EnvelopeFollower.SetAttackTime((float)EnvelopeFollowerAttackTime);
+						EnvelopeFollower.SetReleaseTime((float)EnvelopeFollowerReleaseTime);
+
+						bEnvelopeFollowerSettingsChanged = false;
+					}
+
+					for (int32 FrameIndex = 0; FrameIndex < NumFrames; ++FrameIndex)
+					{
+						EnvelopeFollower.ProcessAudio(BufferToUseForAnalysis[FrameIndex]);
+					}
+
+					CurrentEnvelopeValue = EnvelopeFollower.GetCurrentValue();
+				}
+			}
 		}
 
 		SET_FLOAT_STAT(STAT_MediaUtils_MediaSoundComponentSync, FMath::Abs((Time - OutTime).GetTotalMilliseconds()));
@@ -519,6 +549,26 @@ TArray<FMediaSoundComponentSpectralData> UMediaSoundComponent::GetSpectralData()
 	return TArray<FMediaSoundComponentSpectralData>();
 }
 
+void UMediaSoundComponent::SetEnableEnvelopeFollowing(bool bInEnvelopeFollowing)
+{
+	FScopeLock ScopeLock(&EnvelopeFollowerCriticalSection);
+	bEnvelopeFollowingEnabled = bInEnvelopeFollowing;
+	CurrentEnvelopeValue = 0.0f;
+}
+
+void UMediaSoundComponent::SetEnvelopeFollowingsettings(int32 AttackTimeMsec, int32 ReleaseTimeMsec)
+{
+	FScopeLock ScopeLock(&EnvelopeFollowerCriticalSection);
+	EnvelopeFollowerAttackTime = AttackTimeMsec;
+	EnvelopeFollowerReleaseTime = ReleaseTimeMsec;
+	bEnvelopeFollowerSettingsChanged = true;
+}
+
+float UMediaSoundComponent::GetEnvelopeValue() const
+{
+	return CurrentEnvelopeValue;
+}
+
 /* UMediaSoundComponent implementation
  *****************************************************************************/
 
diff --git a/Engine/Source/Runtime/MediaAssets/Public/MediaSoundComponent.h b/Engine/Source/Runtime/MediaAssets/Public/MediaSoundComponent.h
--- a/Engine/Source/Runtime/MediaAssets/Public/MediaSoundComponent.h
+++ b/Engine/Source/Runtime/MediaAssets/Public/MediaSoundComponent.h
@@ -14,6 +14,8 @@
 #include "UObject/ScriptMacros.h"
 #include "DSP/SpectrumAnalyzer.h"
 #include "DSP/BufferVectorOperations.h"
+#include "DSP/EnvelopeFollower.h"
+
 
 #include "MediaSoundComponent.generated.h"
 
@@ -155,6 +157,18 @@ public:
 	UFUNCTION(BlueprintCallable, Category = "TimeSynth")
 	TArray<FMediaSoundComponentSpectralData> GetSpectralData();
 
+	/** Turns on amplitude envelope following the audio in the media sound component. */
+	UFUNCTION(BlueprintCallable, Category = "Media|MediaSoundComponent")
+	void SetEnableEnvelopeFollowing(bool bInEnvelopeFollowing);
+
+	/** Sets the envelope attack and release times (in ms). */
+	UFUNCTION(BlueprintCallable, Category = "Media|MediaSoundComponent")
+	void SetEnvelopeFollowingsettings(int32 AttackTimeMsec, int32 ReleaseTimeMsec);
+
+	/** Retrieves the current amplitude envelope. */
+	UFUNCTION(BlueprintCallable, Category = "TimeSynth")
+	float GetEnvelopeValue() const;
+
 public:
 
 	void UpdatePlayer();
@@ -257,11 +271,6 @@ private:
 	/** Handle SampleQueue running dry. Ensure audio resumes playback at correct position. */
 	int32 FrameSyncOffset;
 
-	/**
-	 * Sync forward after input audio buffer runs dry due to a hitch or decoder not being able to keep up
-	 * Without this audio will resume playing exactly where it last left off (far behind current player time)
-	 */
-	bool bSyncAudioAfterDropouts;
 
 	/* Time of last sample played. */
 	TAtomic<FTimespan> LastPlaySampleTime;
@@ -276,12 +285,30 @@ private:
 	Audio::FSpectrumAnalyzer SpectrumAnalyzer;
 	Audio::FSpectrumAnalyzerSettings SpectrumAnalyzerSettings;
 
+	Audio::FEnvelopeFollower EnvelopeFollower;
+	int32 EnvelopeFollowerAttackTime;
+	int32 EnvelopeFollowerReleaseTime;
+	float CurrentEnvelopeValue;
+	FCriticalSection EnvelopeFollowerCriticalSection;
+
 	/** Scratch buffer to mix in source audio to from decoder */
 	Audio::AlignedFloatBuffer AudioScratchBuffer;
 
+	/**
+	 * Sync forward after input audio buffer runs dry due to a hitch or decoder not being able to keep up
+	 * Without this audio will resume playing exactly where it last left off (far behind current player time)
+	 */
+	bool bSyncAudioAfterDropouts;
+
 	/** Whether or not spectral analysis is enabled. */
 	bool bSpectralAnalysisEnabled;
 
+	/** Whether or not envelope following is enabled. */
+	bool bEnvelopeFollowingEnabled;
+
+	/** Whether or not envelope follower settings changed. */
+	bool bEnvelopeFollowerSettingsChanged;
+
 private:
 
 	static USoundClass* DefaultMediaSoundClassObject;
