MediaPlaylist will now be added to root set if the owning MediaPlayer is in the disregard for GC set (fixes GC assumption violation crash)

diff --git a/Engine/Source/Runtime/MediaAssets/Private/Assets/MediaPlayer.cpp b/Engine/Source/Runtime/MediaAssets/Private/Assets/MediaPlayer.cpp
--- a/Engine/Source/Runtime/MediaAssets/Private/Assets/MediaPlayer.cpp
+++ b/Engine/Source/Runtime/MediaAssets/Private/Assets/MediaPlayer.cpp
@@ -89,6 +89,23 @@ bool UMediaPlayer::CanPlayUrl(const FString& Url)
 	return PlayerFacade->CanPlayUrl(Url, GetDefault<UMediaSource>());
 }
 
+void UMediaPlayer::SetPlaylistInternal(UMediaPlaylist* InPlaylist)
+{
+	if (Playlist && Playlist != InPlaylist && Playlist->IsRooted())
+	{
+		// To avoid leaking UObjects we need to remove the old playlist from root set 
+		// (which has been most likely rooted because this MediaPlayer is in disregard for GC set)
+		Playlist->RemoveFromRoot();
+	}
+
+	Playlist = InPlaylist;
+
+	if (InPlaylist && GUObjectArray.IsDisregardForGC(this))
+	{
+		// If this MediaPlayer object is in disregard for GC set, add the new Playlist to root set.
+		Playlist->AddToRoot();
+	}
+}
 
 void UMediaPlayer::Close()
 {
@@ -98,8 +115,9 @@ void UMediaPlayer::Close()
 
 	if (!HasAnyFlags(RF_ClassDefaultObject) && !GExitPurge)
 	{
-		Playlist = NewObject<UMediaPlaylist>(GetTransientPackage(), NAME_None, RF_Transactional | RF_Transient);
+		SetPlaylistInternal(NewObject<UMediaPlaylist>(GetTransientPackage(), NAME_None, RF_Transactional | RF_Transient));
 	}
+
 	PlaylistIndex = INDEX_NONE;
 	PlayOnNext = false;
 }
@@ -399,8 +417,8 @@ bool UMediaPlayer::OpenPlaylistIndex(UMediaPlaylist* InPlaylist, int32 Index)
 
 	UE_LOG(LogMediaAssets, Verbose, TEXT("%s.OpenSource %s %i"), *GetFName().ToString(), *InPlaylist->GetFName().ToString(), Index);
 
-	Playlist = InPlaylist;
-
+	SetPlaylistInternal(InPlaylist);
+	
 	if (Index == INDEX_NONE)
 	{
 		return true;
diff --git a/Engine/Source/Runtime/MediaAssets/Public/MediaPlayer.h b/Engine/Source/Runtime/MediaAssets/Public/MediaPlayer.h
--- a/Engine/Source/Runtime/MediaAssets/Public/MediaPlayer.h
+++ b/Engine/Source/Runtime/MediaAssets/Public/MediaPlayer.h
@@ -1080,6 +1080,9 @@ private:
 	/** Callback for when a media event occurred in the player. */
 	void HandlePlayerMediaEvent(EMediaEvent Event);
 
+	/** Sets the playlist and properly handles cases when this MediaPlayer object is in disregard for GC set */
+	void SetPlaylistInternal(UMediaPlaylist* InPlaylist);
+
 private:
 
 	/** An event delegate that is invoked when a media event occurred. */
