Disable GC clustering for UMediaPlayer and UMediaPlaylist

diff --git a/Engine/Source/Runtime/MediaAssets/Public/MediaPlayer.h b/Engine/Source/Runtime/MediaAssets/Public/MediaPlayer.h
--- a/Engine/Source/Runtime/MediaAssets/Public/MediaPlayer.h
+++ b/Engine/Source/Runtime/MediaAssets/Public/MediaPlayer.h
@@ -1076,6 +1076,7 @@ protected:
 	FRotator ViewRotation;
 
 private:
+	virtual bool CanBeInCluster() const override { return false; }
 
 	/** Callback for when a media event occurred in the player. */
 	void HandlePlayerMediaEvent(EMediaEvent Event);
diff --git a/Engine/Source/Runtime/MediaAssets/Public/MediaPlaylist.h b/Engine/Source/Runtime/MediaAssets/Public/MediaPlaylist.h
--- a/Engine/Source/Runtime/MediaAssets/Public/MediaPlaylist.h
+++ b/Engine/Source/Runtime/MediaAssets/Public/MediaPlaylist.h
@@ -148,4 +148,6 @@ protected:
 	/** List of media sources to play. */
 	UPROPERTY(EditAnywhere, Category=Playlist)
 	TArray<UMediaSource*> Items;
+
+	virtual bool CanBeInCluster() const override { return false; }
 };
