Fix Media Player Open Source Latent not triggering output pin in all situations

diff --git a/Engine/Source/Runtime/MediaAssets/Private/Assets/MediaPlayer.cpp b/Engine/Source/Runtime/MediaAssets/Private/Assets/MediaPlayer.cpp
--- a/Engine/Source/Runtime/MediaAssets/Private/Assets/MediaPlayer.cpp
+++ b/Engine/Source/Runtime/MediaAssets/Private/Assets/MediaPlayer.cpp
@@ -942,7 +942,7 @@ public:
 		{
 			UE_LOG(LogMediaAssets, Warning, TEXT("Latent Media: Deleted or Error"));
 			OutSuccess = false;
-			Response.FinishAndTriggerIf(false, ExecutionFunction, OutputLink, CallbackTarget);
+			Response.FinishAndTriggerIf(true, ExecutionFunction, OutputLink, CallbackTarget);
 			return;
 		}
 
@@ -957,7 +957,7 @@ public:
 		{
 			UE_LOG(LogMediaAssets, Verbose, TEXT("Latent Media: Invalid State"));
 			OutSuccess = false;
-			Response.FinishAndTriggerIf(false, ExecutionFunction, OutputLink, CallbackTarget);
+			Response.FinishAndTriggerIf(true, ExecutionFunction, OutputLink, CallbackTarget);
 			return;
 		}
 		else
