Merged: Expose last audio and video sample processed times to blueprint

diff --git a/Engine/Source/Runtime/MediaAssets/Private/Assets/MediaPlayer.cpp b/Engine/Source/Runtime/MediaAssets/Private/Assets/MediaPlayer.cpp
--- a/Engine/Source/Runtime/MediaAssets/Private/Assets/MediaPlayer.cpp
+++ b/Engine/Source/Runtime/MediaAssets/Private/Assets/MediaPlayer.cpp
@@ -225,6 +225,15 @@ FTimespan UMediaPlayer::GetTime() const
 	return PlayerFacade->GetTime();
 }
 
+FTimespan UMediaPlayer::GetLastAudioSampleProcessedTime() const
+{
+	return PlayerFacade->GetLastAudioSampleProcessedTime();
+}
+
+FTimespan UMediaPlayer::GetLastVideoSampleProcessedTime() const
+{
+	return PlayerFacade->GetLastVideoSampleProcessedTime();
+}
 
 FText UMediaPlayer::GetTrackDisplayName(EMediaPlayerTrack TrackType, int32 TrackIndex) const
 {
diff --git a/Engine/Source/Runtime/MediaAssets/Public/MediaPlayer.h b/Engine/Source/Runtime/MediaAssets/Public/MediaPlayer.h
--- a/Engine/Source/Runtime/MediaAssets/Public/MediaPlayer.h
+++ b/Engine/Source/Runtime/MediaAssets/Public/MediaPlayer.h
@@ -285,6 +285,22 @@ public:
 	UFUNCTION(BlueprintCallable, Category="Media|MediaPlayer")
 	FTimespan GetTime() const;
 
+	/**
+	 * Get time of last audio sample decoded
+	 *
+	 * @return Time of last audio sample decoded.
+	 */
+	UFUNCTION(BlueprintCallable, Category="Media|MediaPlayer")
+	FTimespan GetLastAudioSampleProcessedTime() const;
+
+	/**
+	 * Get time of last video sample decoded
+	 *
+	 * @return Time of last video sample decoded.
+	 */
+	UFUNCTION(BlueprintCallable, Category="Media|MediaPlayer")
+	FTimespan GetLastVideoSampleProcessedTime() const;
+
 	/**
 	 * Get the human readable name of the specified track.
 	 *
