Media player can now be created before media module is loaded

diff --git a/Engine/Source/Runtime/MediaAssets/Private/Assets/MediaPlayer.cpp b/Engine/Source/Runtime/MediaAssets/Private/Assets/MediaPlayer.cpp
--- a/Engine/Source/Runtime/MediaAssets/Private/Assets/MediaPlayer.cpp
+++ b/Engine/Source/Runtime/MediaAssets/Private/Assets/MediaPlayer.cpp
@@ -44,6 +44,7 @@ UMediaPlayer::UMediaPlayer(const FObjectInitializer& ObjectInitializer)
 	, ViewRotation(FRotator::ZeroRotator)
 	, PlayerGuid(FGuid::NewGuid())
 	, PlayOnNext(false)
+	, RegisteredWithMediaModule(false)
 #if WITH_EDITORONLY_DATA
 	, AffectedByPIEHandling(true)
 	, WasPlayingInPIE(false)
@@ -384,6 +385,7 @@ bool UMediaPlayer::Next()
 	}
 
 	PlayOnNext |= PlayerFacade->IsPlaying();
+	RegisterWithMediaModule();
 
 	while (RemainingAttempts-- > 0)
 	{
@@ -453,6 +455,7 @@ bool UMediaPlayer::OpenPlaylistIndex(UMediaPlaylist* InPlaylist, int32 Index)
 		return false;
 	}
 
+	RegisterWithMediaModule();
 	return PlayerFacade->Open(MediaSource->GetUrl(), MediaSource);
 }
 
@@ -480,6 +483,7 @@ bool UMediaPlayer::OpenSourceInternal(UMediaSource* MediaSource, const FMediaPla
 	PlayOnNext |= PlayerFacade->IsPlaying();
 	Playlist->GetNext(PlaylistIndex);
 
+	RegisterWithMediaModule();
 	return PlayerFacade->Open(MediaSource->GetUrl(), MediaSource, PlayerOptions);
 }
 
@@ -537,6 +541,7 @@ bool UMediaPlayer::Previous()
 	}
 
 	PlayOnNext |= PlayerFacade->IsPlaying();
+	RegisterWithMediaModule();
 
 	while (--RemainingAttempts >= 0)
 	{
@@ -760,6 +765,11 @@ void UMediaPlayer::PostInitProperties()
 
 void UMediaPlayer::RegisterWithMediaModule()
 {
+	if (RegisteredWithMediaModule)
+	{
+		return;
+	}
+
 	static const FName MediaModuleName("Media");
 	IMediaModule* MediaModule = nullptr;
 	if (IsInGameThread())
@@ -777,6 +787,7 @@ void UMediaPlayer::RegisterWithMediaModule()
 	{
 		MediaModule->GetClock().AddSink(PlayerFacade.ToSharedRef());
 		MediaModule->GetTicker().AddTickable(PlayerFacade.ToSharedRef());
+		RegisteredWithMediaModule = true;
 	}
 	else
 	{
diff --git a/Engine/Source/Runtime/MediaAssets/Public/MediaPlayer.h b/Engine/Source/Runtime/MediaAssets/Public/MediaPlayer.h
--- a/Engine/Source/Runtime/MediaAssets/Public/MediaPlayer.h
+++ b/Engine/Source/Runtime/MediaAssets/Public/MediaPlayer.h
@@ -1155,6 +1155,9 @@ private:
 	/** Automatically start playback of next item in play list. */
 	bool PlayOnNext;
 
+	/** Has registered with the media module */
+	bool RegisteredWithMediaModule;
+
 #if WITH_EDITORONLY_DATA
 public:
 	/** Whether this player should stop when entering or exiting PIE. */
