Open Media Latent function now triggers output pin after open, play on open, or seek completed.

diff --git a/Engine/Source/Runtime/MediaAssets/Private/Assets/MediaPlayer.cpp b/Engine/Source/Runtime/MediaAssets/Private/Assets/MediaPlayer.cpp
--- a/Engine/Source/Runtime/MediaAssets/Private/Assets/MediaPlayer.cpp
+++ b/Engine/Source/Runtime/MediaAssets/Private/Assets/MediaPlayer.cpp
@@ -14,6 +14,7 @@
 #include "LatentActions.h"
 #include "MediaPlayerFacade.h"
 #include "MediaPlayerOptions.h"
+#include "MediaHelpers.h"
 #include "Misc/App.h"
 #include "Misc/Paths.h"
 #include "Misc/ScopeLock.h"
@@ -365,6 +366,11 @@ bool UMediaPlayer::IsPreparing() const
 }
 
 
+bool UMediaPlayer::IsClosed() const
+{
+	return PlayerFacade->IsClosed();
+}
+
 bool UMediaPlayer::IsReady() const
 {
 	UE_LOG(LogMediaAssets, VeryVerbose, TEXT("%s.IsReady"), *GetFName().ToString());
@@ -933,6 +939,10 @@ public:
 	float TimeRemaining;
 	bool& OutSuccess;
 	bool bSawError;
+	bool bSawMediaOpened;
+	bool bSawMediaOpenFailed;
+	bool bSawSeekCompleted;
+	FMediaPlayerOptions Options;
 	FString URL;
 
 	FLatentOpenMediaSourceAction(const FLatentActionInfo& LatentInfo, UMediaPlayer* InMediaPlayer, UMediaSource* InMediaSource, const FMediaPlayerOptions& InOptions, bool& InSuccess)
@@ -943,55 +953,151 @@ public:
 		, TimeRemaining(10.0)
 		, OutSuccess(InSuccess)
 		, bSawError(false)
+		, bSawMediaOpened(false)
+		, bSawMediaOpenFailed(false)
+		, bSawSeekCompleted(false)
+		, Options(InOptions)
 	{
 		if (InMediaSource)
 		{
 			URL = InMediaSource->GetUrl();
-		}
 
-		if (!InMediaPlayer->OpenSourceWithOptions(InMediaSource, InOptions))
+			InMediaPlayer->OnMediaEvent().AddRaw(this, &FLatentOpenMediaSourceAction::HandleMediaPlayerEvent);
+			if (!InMediaPlayer->OpenSourceWithOptions(InMediaSource, InOptions))
+			{
+				UE_LOG(LogMediaAssets, Warning, TEXT("Open Media Latent: Failed initial open: %s"), *URL);
+				bSawError = true;
+			}
+		}
+		else
 		{
-			UE_LOG(LogMediaAssets, Verbose, TEXT("Latent Media: After initial open started"));
+			UE_LOG(LogMediaAssets, Warning, TEXT("Open Media Latent: Failed initial open because no media source given"));
 			bSawError = true;
 		}
 	}
 
+	virtual ~FLatentOpenMediaSourceAction()
+	{
+		UnregisterMediaEvent();
+	}
+
+	void UnregisterMediaEvent()
+	{
+		if (MediaPlayer.IsValid())
+		{
+			MediaPlayer->OnMediaEvent().RemoveAll(this);
+		}
+	}
+
+	void HandleMediaPlayerEvent(EMediaEvent Event)
+	{
+		UE_LOG(LogMediaAssets, Verbose, TEXT("Open Media Latent: Saw event: %s"), *MediaUtils::EventToString(Event));
+
+		switch (Event)
+		{
+		case EMediaEvent::MediaOpened:
+			bSawMediaOpened = true;
+			break;
+		case EMediaEvent::MediaOpenFailed:
+			bSawMediaOpenFailed = true;
+			break;
+		case EMediaEvent::SeekCompleted:
+			bSawSeekCompleted = true;
+			break;
+		}
+	}
+
+	void FailedOperation(FLatentResponse& Response)
+	{
+		OutSuccess = false;
+		Response.FinishAndTriggerIf(true, ExecutionFunction, OutputLink, CallbackTarget);
+	}
+
 	virtual void UpdateOperation(FLatentResponse& Response) override
 	{
-		if (bSawError || !MediaPlayer.IsValid() || MediaPlayer.IsStale() || MediaPlayer->HasError())
+		if (bSawMediaOpenFailed)
 		{
-			UE_LOG(LogMediaAssets, Warning, TEXT("Latent Media: Deleted or Error"));
-			OutSuccess = false;
-			Response.FinishAndTriggerIf(true, ExecutionFunction, OutputLink, CallbackTarget);
+			UE_LOG(LogMediaAssets, Warning, TEXT("Open Media Latent: Saw media open failed event. %s"), *URL);
+			FailedOperation(Response);
 			return;
 		}
 
-		if (MediaPlayer->IsReady())
+		if (!MediaPlayer.IsValid() || MediaPlayer.IsStale())
 		{
-			UE_LOG(LogMediaAssets, Verbose, TEXT("Latent Media: IsReady"));
-			OutSuccess = true;
-			Response.FinishAndTriggerIf(true, ExecutionFunction, OutputLink, CallbackTarget);
+			UE_LOG(LogMediaAssets, Warning, TEXT("Open Media Latent: Media player object was deleted. %s"), *URL);
+			FailedOperation(Response);
 			return;
 		}
-		else if (!MediaPlayer->IsPreparing())
+
+		if (bSawError || MediaPlayer->HasError())
 		{
-			UE_LOG(LogMediaAssets, Verbose, TEXT("Latent Media: Invalid State"));
-			OutSuccess = false;
-			Response.FinishAndTriggerIf(true, ExecutionFunction, OutputLink, CallbackTarget);
+			UE_LOG(LogMediaAssets, Warning, TEXT("Open Media Latent: Media player is in Error state. %s"), *URL);
+			FailedOperation(Response);
 			return;
 		}
+
+		if (MediaPlayer->IsClosed())
+		{
+			UE_LOG(LogMediaAssets, Warning, TEXT("Open Media Latent: Media player is closed. %s"), *URL);
+			FailedOperation(Response);
+			return;
+		}
+
+		if (MediaPlayer->IsPreparing())
+		{
+			UE_LOG(LogMediaAssets, Verbose, TEXT("Open Media Latent: Is preparing ..."));
+		}
+		else if (MediaPlayer->IsReady())
+		{
+			UE_LOG(LogMediaAssets, Verbose, TEXT("Open Media Latent: IsReady() ... %s"), *URL);
+
+			if (bSawMediaOpened)
+			{
+				if (!Options.SeekTime.IsZero())
+				{
+					if (bSawSeekCompleted)
+					{
+						UE_LOG(LogMediaAssets, Verbose, TEXT("Open Media Latent: Triggering output pin after seek completed. Success: %d, %s"), OutSuccess, *URL);
+						OutSuccess = true;
+						Response.FinishAndTriggerIf(true, ExecutionFunction, OutputLink, CallbackTarget);
+						return;
+					}
+					else
+					{
+						if (Options.SeekTime < FTimespan::FromSeconds(0) || Options.SeekTime > MediaPlayer->GetDuration())
+						{
+							UE_LOG(LogMediaAssets, Warning, TEXT("Open Media Latent: Media player seeking to time out of bounds. Seek: %s, Duration: %s, URL: %s"), 
+								*Options.SeekTime.ToString(), *MediaPlayer->GetDuration().ToString(), *URL);
+							FailedOperation(Response);
+							return;
+						}
+						UE_LOG(LogMediaAssets, Verbose, TEXT("Open Media Latent: Waiting for seek completed event ..."));
+					}
+				}
+				else
+				{
+					UE_LOG(LogMediaAssets, Verbose, TEXT("Open Media Latent: Triggering output pin after opened event. Success: %d, %s"), OutSuccess, *URL);
+					OutSuccess = true;
+					Response.FinishAndTriggerIf(true, ExecutionFunction, OutputLink, CallbackTarget);
+					return;
+				}
+			}
+			else
+			{
+				UE_LOG(LogMediaAssets, Verbose, TEXT("Open Media Latent: Waiting for opened event ..."));
+			}
+		}
 		else
 		{
-			UE_LOG(LogMediaAssets, Verbose, TEXT("Latent Media: Waiting for open ..."));
+			UE_LOG(LogMediaAssets, Verbose, TEXT("Open Media Latent: Waiting for IsReady() ..."));
 		}
 
 		// Timed out
 		TimeRemaining -= Response.ElapsedTime();
 		if (TimeRemaining <= 0.0f)
 		{
-			UE_LOG(LogMediaAssets, Verbose, TEXT("Latent Media: Timed out."));
-			OutSuccess = false;
-			Response.FinishAndTriggerIf(true, ExecutionFunction, OutputLink, CallbackTarget);
+			UE_LOG(LogMediaAssets, Warning, TEXT("Open Media Latent: Timed out. %s"), *URL);
+			FailedOperation(Response);
 			return;
 		}
 	}
diff --git a/Engine/Source/Runtime/MediaAssets/Public/MediaPlayer.h b/Engine/Source/Runtime/MediaAssets/Public/MediaPlayer.h
--- a/Engine/Source/Runtime/MediaAssets/Public/MediaPlayer.h
+++ b/Engine/Source/Runtime/MediaAssets/Public/MediaPlayer.h
@@ -492,6 +492,14 @@ public:
 	UFUNCTION(BlueprintCallable, Category="Media|MediaPlayer")
 	bool IsPreparing() const;
 
+	/**
+	 * Whether media is currently closed.
+	 *
+	 * @return true if media is closed, false otherwise.
+	 */
+	UFUNCTION(BlueprintCallable, Category="Media|MediaPlayer")
+	bool IsClosed() const;
+
 	/**
 	 * Check whether media is ready for playback.
 	 *
@@ -589,11 +597,11 @@ public:
 	/**
 	 * Open the specified media source with options using a latent action.
 	 *
-	 * A result of true indicates that the player opened the media source successfully.
+	 * A result of true indicates that the player successfully completed all requested operations.
 	 *
 	 * @param MediaSource The media source to open.
 	 * @param Options The media player options to apply.
-	 * @param bSuccess  Player opened the media.
+	 * @param bSuccess  All requested operations have completed successfully.
 	 * @see Close, OpenFile, OpenPlaylist, OpenPlaylistIndex, OpenUrl, Reopen
 	 */
 	UFUNCTION(BlueprintCallable, Category = "Media|MediaPlayer", meta = (Latent, LatentInfo = "LatentInfo", WorldContext = "WorldContextObject"))
diff --git a/Engine/Source/Runtime/MediaUtils/Private/MediaPlayerFacade.cpp b/Engine/Source/Runtime/MediaUtils/Private/MediaPlayerFacade.cpp
--- a/Engine/Source/Runtime/MediaUtils/Private/MediaPlayerFacade.cpp
+++ b/Engine/Source/Runtime/MediaUtils/Private/MediaPlayerFacade.cpp
@@ -412,6 +412,10 @@ bool FMediaPlayerFacade::IsPreparing() const
 	return Player.IsValid() && (Player->GetControls().GetState() == EMediaState::Preparing);
 }
 
+bool FMediaPlayerFacade::IsClosed() const
+{
+	return Player.IsValid() && (Player->GetControls().GetState() == EMediaState::Closed);
+}
 
 bool FMediaPlayerFacade::IsReady() const
 {
diff --git a/Engine/Source/Runtime/MediaUtils/Public/MediaPlayerFacade.h b/Engine/Source/Runtime/MediaUtils/Public/MediaPlayerFacade.h
--- a/Engine/Source/Runtime/MediaUtils/Public/MediaPlayerFacade.h
+++ b/Engine/Source/Runtime/MediaUtils/Public/MediaPlayerFacade.h
@@ -472,6 +472,13 @@ public:
 	 */
 	bool IsPreparing() const;
 
+	/**
+	 * Whether media is currently closed.
+	 *
+	 * @return true if media is closed, false otherwise.
+	 */
+	bool IsClosed() const;
+
 	/**
 	 * Whether media is ready for playback.
 	 *
