Merging //UE4/Dev-Main to Dev-Sequencer (//UE4/Dev-Sequencer)

diff --git a/Engine/Plugins/Experimental/PythonScriptPlugin/Source/PythonScriptPlugin/Private/PyReferenceCollector.cpp b/Engine/Plugins/Experimental/PythonScriptPlugin/Source/PythonScriptPlugin/Private/PyReferenceCollector.cpp
--- a/Engine/Plugins/Experimental/PythonScriptPlugin/Source/PythonScriptPlugin/Private/PyReferenceCollector.cpp
+++ b/Engine/Plugins/Experimental/PythonScriptPlugin/Source/PythonScriptPlugin/Private/PyReferenceCollector.cpp
@@ -5,59 +5,10 @@
 #include "PyWrapperDelegate.h"
 #include "UObject/UnrealType.h"
 #include "UObject/UObjectHash.h"
+#include "UObject/PurgingReferenceCollector.h"
 
 #if WITH_PYTHON
 
-/** Reference collector that will purge (null) any references to the given set of objects (as if they'd bee marked PendingKill) */
-class FPyPurgingReferenceCollector : public FReferenceCollector
-{
-public:
-	bool HasObjectToPurge() const
-	{
-		return ObjectsToPurge.Num() > 0;
-	}
-
-	void AddObjectToPurge(const UObject* Object)
-	{
-		ObjectsToPurge.Add(Object);
-	}
-
-	virtual bool IsIgnoringArchetypeRef() const override
-	{
-		return false;
-	}
-
-	virtual bool IsIgnoringTransient() const override
-	{
-		return false;
-	}
-
-protected:
-	virtual void HandleObjectReference(UObject*& Object, const UObject* ReferencingObject, const UProperty* ReferencingProperty) override
-	{
-		ConditionalPurgeObject(Object);
-	}
-
-	virtual void HandleObjectReferences(UObject** InObjects, const int32 ObjectNum, const UObject* InReferencingObject, const UProperty* InReferencingProperty) override
-	{
-		for (int32 ObjectIndex = 0; ObjectIndex < ObjectNum; ++ObjectIndex)
-		{
-			UObject*& Object = InObjects[ObjectIndex];
-			ConditionalPurgeObject(Object);
-		}
-	}
-
-	void ConditionalPurgeObject(UObject*& Object)
-	{
-		if (ObjectsToPurge.Contains(Object))
-		{
-			Object = nullptr;
-		}
-	}
-
-	TSet<const UObject*> ObjectsToPurge;
-};
-
 FPyReferenceCollector& FPyReferenceCollector::Get()
 {
 	static FPyReferenceCollector Instance;
@@ -93,7 +44,7 @@ void FPyReferenceCollector::PurgeUnrealObjectReferences(const UObject* InObject,
 
 void FPyReferenceCollector::PurgeUnrealObjectReferences(const TArrayView<const UObject*>& InObjects, const bool bIncludeInnerObjects)
 {
-	FPyPurgingReferenceCollector PurgingReferenceCollector;
+	FPurgingReferenceCollector PurgingReferenceCollector;
 
 	for (const UObject* Object : InObjects)
 	{
diff --git a/Engine/Source/Editor/Blutility/Private/BlutilityModule.cpp b/Engine/Source/Editor/Blutility/Private/BlutilityModule.cpp
--- a/Engine/Source/Editor/Blutility/Private/BlutilityModule.cpp
+++ b/Engine/Source/Editor/Blutility/Private/BlutilityModule.cpp
@@ -5,8 +5,9 @@
 #include "Widgets/DeclarativeSyntaxSupport.h"
 #include "Framework/Docking/TabManager.h"
 #include "IBlutilityModule.h"
+#include "EditorUtilityWidget.h"
 #include "EditorUtilityBlueprint.h"
-
+#include "GlobalEditorUtilityBase.h"
 
 #include "AssetToolsModule.h"
 #include "PropertyEditorModule.h"
@@ -30,7 +31,8 @@
 #include "LevelEditor.h"
 #include "Editor.h"
 #include "UnrealEdMisc.h"
-
+#include "EditorSupportDelegates.h"
+#include "UObject/PurgingReferenceCollector.h"
 
 #define LOCTEXT_NAMESPACE "AssetTypeActions"
 
@@ -91,6 +93,8 @@ public:
 		FLevelEditorModule& LevelEditorModule = FModuleManager::LoadModuleChecked<FLevelEditorModule>(TEXT("LevelEditor"));
 		LevelEditorModule.OnTabManagerChanged().AddRaw(this, &FBlutilityModule::ReinitializeUIs);
 		LevelEditorModule.OnMapChanged().AddRaw(this, &FBlutilityModule::OnMapChanged);
+
+		FEditorSupportDelegates::PrepareToCleanseEditorObject.AddRaw(this, &FBlutilityModule::OnPrepareToCleanseEditorObject);
 	}
 
 	void ReinitializeUIs()
@@ -184,6 +188,8 @@ public:
 			PropertyModule.UnregisterCustomClassLayout("GlobalEditorUtilityBase");
 			PropertyModule.NotifyCustomizationModuleChanged();
 		}
+
+		FEditorSupportDelegates::PrepareToCleanseEditorObject.RemoveAll(this);
 	}
 
 	virtual bool IsBlutility( const UBlueprint* Blueprint ) const override
@@ -244,6 +250,41 @@ protected:
 		}
 	}
 
+	void OnPrepareToCleanseEditorObject(UObject* InObject)
+	{
+		// Gather the list of live Editor Utility instances to purge references from
+		TArray<UObject*> EditorUtilityInstances;
+		ForEachObjectOfClass(UEditorUtilityWidget::StaticClass(), [&EditorUtilityInstances](UObject* InEditorUtilityInstance)
+		{
+			EditorUtilityInstances.Add(InEditorUtilityInstance);
+		});
+		ForEachObjectOfClass(UGlobalEditorUtilityBase::StaticClass(), [&EditorUtilityInstances](UObject* InEditorUtilityInstance)
+		{
+			EditorUtilityInstances.Add(InEditorUtilityInstance);
+		});
+
+		if (EditorUtilityInstances.Num() > 0)
+		{
+			// Build up the complete list of objects to purge
+			FPurgingReferenceCollector PurgingReferenceCollector;
+			PurgingReferenceCollector.AddObjectToPurge(InObject);
+			ForEachObjectWithOuter(InObject, [&PurgingReferenceCollector](UObject* InInnerObject)
+			{
+				PurgingReferenceCollector.AddObjectToPurge(InInnerObject);
+			}, true);
+
+			// Run the purge for each Editor Utility instance
+			FReferenceCollectorArchive& PurgingReferenceCollectorArchive = PurgingReferenceCollector.GetVerySlowReferenceCollectorArchive();
+			for (UObject* EditorUtilityInstance : EditorUtilityInstances)
+			{
+				PurgingReferenceCollectorArchive.SetSerializingObject(EditorUtilityInstance);
+				EditorUtilityInstance->Serialize(PurgingReferenceCollectorArchive);
+				EditorUtilityInstance->CallAddReferencedObjects(PurgingReferenceCollector);
+				PurgingReferenceCollectorArchive.SetSerializingObject(nullptr);
+			}
+		}
+	}
+
 	/** Scripted Editor Widgets workspace menu item */
 	TSharedPtr<class FWorkspaceItem> ScriptedEditorWidgetsGroup;
 
diff --git a/Engine/Source/Editor/UnrealEd/Private/Lightmass/LightmassLandscapeRender.cpp b/Engine/Source/Editor/UnrealEd/Private/Lightmass/LightmassLandscapeRender.cpp
--- a/Engine/Source/Editor/UnrealEd/Private/Lightmass/LightmassLandscapeRender.cpp
+++ b/Engine/Source/Editor/UnrealEd/Private/Lightmass/LightmassLandscapeRender.cpp
@@ -70,7 +70,7 @@ void RenderLandscapeMaterialForLightmass(const FLandscapeStaticLightingMesh* Lan
 			                                    (SubsectionY >= 0 && SubsectionY < NumSubsections ? 1 : 0));
 
 			const FVector2D BasePosition = PatchExpandOffset + FVector2D(SubsectionX, SubsectionY) * PositionScale;
-			const FVector2D BaseLayerCoords = FVector2D(UVSubsection) * LayerScale;
+			const FVector2D BaseLayerCoords = FVector2D(LandscapeComponent->SectionBaseX, LandscapeComponent->SectionBaseY) + FVector2D(UVSubsection) * LayerScale;
 			const FVector2D BaseWeightmapCoords = WeightmapBias + FVector2D(UVSubsection) * WeightmapSubsection;
 
 			int32 Index = Vertices.Add(FDynamicMeshVertex(FVector(BasePosition /*FVector2D(0, 0) * PositionScale*/, 0), FVector(BaseLayerCoords /*FVector2D(0, 0) * UVScale * LayerScale*/, 0), BaseWeightmapCoords /*FVector2D(0, 0) * UVScale * WeightmapScale*/));
diff --git a/Engine/Source/Programs/UnrealBuildTool/Configuration/UEBuildModuleCPP.cs b/Engine/Source/Programs/UnrealBuildTool/Configuration/UEBuildModuleCPP.cs
--- a/Engine/Source/Programs/UnrealBuildTool/Configuration/UEBuildModuleCPP.cs
+++ b/Engine/Source/Programs/UnrealBuildTool/Configuration/UEBuildModuleCPP.cs
@@ -64,12 +64,14 @@ namespace UnrealBuildTool
 		{
 			base.GetReferencedDirectories(Directories);
 
-			if(SourceDirectories == null)
+			if(!Rules.bUsePrecompiled)
 			{
-				throw new BuildException("GetReferencedDirectories() should not be called before building.");
+				if(SourceDirectories == null)
+				{
+					throw new BuildException("GetReferencedDirectories() should not be called before building.");
+				}
+				Directories.UnionWith(SourceDirectories);
 			}
-
-			Directories.UnionWith(SourceDirectories);
 		}
 
 		/// <summary>
diff --git a/Engine/Source/Runtime/Core/Private/Unix/UnixErrorOutputDevice.cpp b/Engine/Source/Runtime/Core/Private/Unix/UnixErrorOutputDevice.cpp
--- a/Engine/Source/Runtime/Core/Private/Unix/UnixErrorOutputDevice.cpp
+++ b/Engine/Source/Runtime/Core/Private/Unix/UnixErrorOutputDevice.cpp
@@ -54,7 +54,13 @@ void FUnixErrorOutputDevice::Serialize(const TCHAR* Msg, ELogVerbosity::Type Ver
 #if PLATFORM_EXCEPTIONS_DISABLED
 		UE_DEBUG_BREAK();
 #endif
-		ReportAssert(Msg, 0);
+		// Generate the portable callstack. For asserts, we ignore the following frames:
+		//     FDebug::AssertFailed()
+		//   [ FOutputDevice::Logf() ] - force-inlined; ignored
+		//     FOutputDevice::LogfImpl()
+		//     FWindowsErrorOutputDevice::Serialize()
+		const int32 NumStackFramesToIgnore = 3;
+		ReportAssert(Msg, NumStackFramesToIgnore);
 	}
 	else
 	{
diff --git a/Engine/Source/Runtime/Core/Private/Unix/UnixPlatformCrashContext.cpp b/Engine/Source/Runtime/Core/Private/Unix/UnixPlatformCrashContext.cpp
--- a/Engine/Source/Runtime/Core/Private/Unix/UnixPlatformCrashContext.cpp
+++ b/Engine/Source/Runtime/Core/Private/Unix/UnixPlatformCrashContext.cpp
@@ -101,6 +101,12 @@ void FUnixCrashContext::InitFromSignal(int32 InSignal, siginfo_t* InInfo, void*
 	Context = reinterpret_cast< ucontext_t* >( InContext );
 
 	FCString::Strcat(SignalDescription, ARRAY_COUNT( SignalDescription ) - 1, *DescribeSignal(Signal, Info, Context));
+
+	// Retrieve NumStackFramesToIgnore from signal data
+	if (Info && Info->si_value.sival_int != 0)
+	{
+		SetNumMinidumpFramesToIgnore(Info->si_value.sival_int);
+	}
 }
 
 void FUnixCrashContext::InitFromEnsureHandler(const TCHAR* EnsureMessage, const void* CrashAddress)
@@ -292,7 +298,7 @@ void FUnixCrashContext::CaptureStackTrace()
 		ANSICHAR* StackTrace = (ANSICHAR*) FMemory::Malloc( StackTraceSize );
 		StackTrace[0] = 0;
 
-		int32 IgnoreCount = 0;
+		int32 IgnoreCount = NumMinidumpFramesToIgnore;
 		CapturePortableCallStack(IgnoreCount, this);
 
 		// Walk the stack and dump it to the allocated memory (do not ignore any stack frames to be consistent with check()/ensure() handling)
diff --git a/Engine/Source/Runtime/Core/Private/Unix/UnixPlatformStackWalk.cpp b/Engine/Source/Runtime/Core/Private/Unix/UnixPlatformStackWalk.cpp
--- a/Engine/Source/Runtime/Core/Private/Unix/UnixPlatformStackWalk.cpp
+++ b/Engine/Source/Runtime/Core/Private/Unix/UnixPlatformStackWalk.cpp
@@ -787,7 +787,17 @@ thread_local const TCHAR* GAssertErrorMessage = nullptr;
 void ReportAssert(const TCHAR* ErrorMessage, int NumStackFramesToIgnore)
 {
 	GAssertErrorMessage = ErrorMessage;
-	FPlatformMisc::RaiseException(1);
+
+	// Store NumStackFramesToIgnore in signal data	
+	sigval UserData;
+	UserData.sival_int = NumStackFramesToIgnore + 2; // +2 for this function and sigqueue()
+	sigqueue(getpid(),  SIGSEGV, UserData);
+	
+	// Make sure we never return
+	for (;;)
+	{
+		FPlatformProcess::Sleep(60.0f);
+	}
 }
 
 static FCriticalSection EnsureLock;
diff --git a/Engine/Source/Runtime/CoreUObject/Public/UObject/PurgingReferenceCollector.h b/Engine/Source/Runtime/CoreUObject/Public/UObject/PurgingReferenceCollector.h
new file mode 100644
--- /dev/null
+++ b/Engine/Source/Runtime/CoreUObject/Public/UObject/PurgingReferenceCollector.h
@@ -0,0 +1,55 @@
+// Copyright 1998-2019 Epic Games, Inc. All Rights Reserved.
+
+#pragma once
+
+#include "UObject/UnrealType.h"
+
+/** Reference collector that will purge (null) any references to the given set of objects (as if they'd been marked PendingKill) */
+class FPurgingReferenceCollector : public FReferenceCollector
+{
+public:
+	bool HasObjectToPurge() const
+	{
+		return ObjectsToPurge.Num() > 0;
+	}
+
+	void AddObjectToPurge(const UObject* Object)
+	{
+		ObjectsToPurge.Add(Object);
+	}
+
+	virtual bool IsIgnoringArchetypeRef() const override
+	{
+		return false;
+	}
+
+	virtual bool IsIgnoringTransient() const override
+	{
+		return false;
+	}
+
+protected:
+	virtual void HandleObjectReference(UObject*& Object, const UObject* ReferencingObject, const UProperty* ReferencingProperty) override
+	{
+		ConditionalPurgeObject(Object);
+	}
+
+	virtual void HandleObjectReferences(UObject** InObjects, const int32 ObjectNum, const UObject* InReferencingObject, const UProperty* InReferencingProperty) override
+	{
+		for (int32 ObjectIndex = 0; ObjectIndex < ObjectNum; ++ObjectIndex)
+		{
+			UObject*& Object = InObjects[ObjectIndex];
+			ConditionalPurgeObject(Object);
+		}
+	}
+
+	void ConditionalPurgeObject(UObject*& Object)
+	{
+		if (ObjectsToPurge.Contains(Object))
+		{
+			Object = nullptr;
+		}
+	}
+
+	TSet<const UObject*> ObjectsToPurge;
+};
diff --git a/Engine/Source/Runtime/Engine/Private/KismetSystemLibrary.cpp b/Engine/Source/Runtime/Engine/Private/KismetSystemLibrary.cpp
--- a/Engine/Source/Runtime/Engine/Private/KismetSystemLibrary.cpp
+++ b/Engine/Source/Runtime/Engine/Private/KismetSystemLibrary.cpp
@@ -302,10 +302,15 @@ void UKismetSystemLibrary::ExecuteConsoleCommand(UObject* WorldContextObject, co
 {
 	// First, try routing through the primary player
 	APlayerController* TargetPC = Player ? Player : UGameplayStatics::GetPlayerController(WorldContextObject, 0);
-	if( TargetPC )
+	if (TargetPC)
 	{
 		TargetPC->ConsoleCommand(Command, true);
 	}
+	else if (GEngine)
+	{
+		UWorld* World = GEngine->GetWorldFromContextObject(WorldContextObject, EGetWorldErrorMode::ReturnNull);
+		GEngine->Exec(World, *Command);
+	}
 }
 
 float UKismetSystemLibrary::GetConsoleVariableFloatValue(const FString& VariableName)
diff --git a/Engine/Source/Runtime/Engine/Public/PrimitiveSceneProxy.h b/Engine/Source/Runtime/Engine/Public/PrimitiveSceneProxy.h
--- a/Engine/Source/Runtime/Engine/Public/PrimitiveSceneProxy.h
+++ b/Engine/Source/Runtime/Engine/Public/PrimitiveSceneProxy.h
@@ -679,11 +679,12 @@ public:
 	 * @param InView - Current View
  	 * @param InViewLODScale - View LOD scale
   	 * @param InCustomDataMemStack - MemStack to allocate the custom data
-   	 * @param InIsStaticRelevant - Tell us if it was called in a static of dynamic relevancy context
+	 * @param InIsStaticRelevant - Tell us if it was called in a static of dynamic relevancy context
+	 * @param InIsShadowOnly - Tell us if we are creating in the shadow context
    	 * @param InVisiblePrimitiveLODMask - Calculated LODMask for visibile primitive in static relevancy
    	 * @param InMeshScreenSizeSquared - Computed mesh batch screen size, passed to prevent recalculation
 	 */
-	ENGINE_API virtual void* InitViewCustomData(const FSceneView& InView, float InViewLODScale, FMemStackBase& InCustomDataMemStack, bool InIsStaticRelevant = false, const struct FLODMask* InVisiblePrimitiveLODMask = nullptr, float InMeshScreenSizeSquared = -1.0f) { return nullptr; }
+	ENGINE_API virtual void* InitViewCustomData(const FSceneView& InView, float InViewLODScale, FMemStackBase& InCustomDataMemStack, bool InIsStaticRelevant, bool InIsShadowOnly, const struct FLODMask* InVisiblePrimitiveLODMask = nullptr, float InMeshScreenSizeSquared = -1.0f) { return nullptr; }
 	
 	/**
 	 * Called during post visibility and shadow setup, just before the frame is rendered. It can be used to update custom data that had a dependency between them.
diff --git a/Engine/Source/Runtime/Landscape/Private/LandscapeRender.cpp b/Engine/Source/Runtime/Landscape/Private/LandscapeRender.cpp
--- a/Engine/Source/Runtime/Landscape/Private/LandscapeRender.cpp
+++ b/Engine/Source/Runtime/Landscape/Private/LandscapeRender.cpp
@@ -1954,7 +1954,7 @@ int8 FLandscapeComponentSceneProxy::GetLODFromScreenSize(float InScreenSizeSquar
 	return INDEX_NONE;
 }
 
-void* FLandscapeComponentSceneProxy::InitViewCustomData(const FSceneView& InView, float InViewLODScale, FMemStackBase& InCustomDataMemStack, bool InIsStaticRelevant, const FLODMask* InVisiblePrimitiveLODMask, float InMeshScreenSizeSquared)
+void* FLandscapeComponentSceneProxy::InitViewCustomData(const FSceneView& InView, float InViewLODScale, FMemStackBase& InCustomDataMemStack, bool InIsStaticRelevant, bool InIsShadowOnly, const FLODMask* InVisiblePrimitiveLODMask, float InMeshScreenSizeSquared)
 {
 	SCOPE_CYCLE_COUNTER(STAT_LandscapeInitViewCustomData);
 
@@ -1994,6 +1994,8 @@ void* FLandscapeComponentSceneProxy::InitViewCustomData(const FSceneView& InView
 		}
 	}
 
+	LODData->IsShadowOnly = InIsShadowOnly;
+
 	// Mobile use a different way of calculating the Bias
 	if (GetScene().GetFeatureLevel() >= ERHIFeatureLevel::SM4)
 	{
@@ -2138,9 +2140,14 @@ float FLandscapeComponentSceneProxy::GetNeighborLOD(const FSceneView& InView, fl
 
 	if (CustomData != nullptr)
 	{
-		ComputeNeighborCustomDataLOD = false;
 		const FViewCustomDataLOD* LODData = (const FViewCustomDataLOD*)CustomData;
-		NeighborLOD = FMath::Max(LODData->SubSections[DesiredSubSectionIndex].fBatchElementCurrentLOD, InBatchElementCurrentLOD);
+		// Don't use the custom data for neighbor calculation when it is marked shadow only (ie it is not visible in the view)
+		// See UE-69785 for more information
+		if (!LODData->IsShadowOnly)
+		{
+			ComputeNeighborCustomDataLOD = false;
+			NeighborLOD = FMath::Max(LODData->SubSections[DesiredSubSectionIndex].fBatchElementCurrentLOD, InBatchElementCurrentLOD);
+		}
 	}
 
 	if (ComputeNeighborCustomDataLOD)
diff --git a/Engine/Source/Runtime/Landscape/Public/LandscapeRender.h b/Engine/Source/Runtime/Landscape/Public/LandscapeRender.h
--- a/Engine/Source/Runtime/Landscape/Public/LandscapeRender.h
+++ b/Engine/Source/Runtime/Landscape/Public/LandscapeRender.h
@@ -501,6 +501,7 @@ public:
 		FViewCustomDataLOD()
 			: StaticMeshBatchLOD(INDEX_NONE)
 			, UseCombinedMeshBatch(true)
+			, IsShadowOnly(false)
 			, ComponentScreenSize(0.0f)
 			, ShaderCurrentLOD(ForceInitToZero)
 			, LodBias(ForceInitToZero)
@@ -509,6 +510,7 @@ public:
 
 		int8 StaticMeshBatchLOD;
 		bool UseCombinedMeshBatch;
+		bool IsShadowOnly;
 		float ComponentScreenSize;
 		TStaticArray<FViewCustomDataSubSectionLOD, MAX_SUBSECTION_COUNT> SubSections; // We always have at least 1 subsections
 
@@ -677,7 +679,7 @@ public:
 	virtual void OnTransformChanged() override;
 	virtual void CreateRenderThreadResources() override;
 	virtual void OnLevelAddedToWorld() override;
-	virtual void* InitViewCustomData(const FSceneView& InView, float InViewLODScale, FMemStackBase& InCustomDataMemStack, bool InIsStaticRelevant = false, const FLODMask* InVisiblePrimitiveLODMask = nullptr, float InMeshScreenSizeSquared = -1.0f) override;
+	virtual void* InitViewCustomData(const FSceneView& InView, float InViewLODScale, FMemStackBase& InCustomDataMemStack, bool InIsStaticRelevant, bool InIsShadowOnly, const FLODMask* InVisiblePrimitiveLODMask = nullptr, float InMeshScreenSizeSquared = -1.0f) override;
 	virtual void PostInitViewCustomData(const FSceneView& InView, void* InViewCustomData) const override;
 	virtual bool IsUsingCustomLODRules() const override;
 	virtual FLODMask GetCustomLOD(const FSceneView& View, float InViewLODScale, int32 InForcedLODLevel, float& OutScreenSizeSquared) const override;
diff --git a/Engine/Source/Runtime/MediaAssets/Private/Assets/MediaPlayer.cpp b/Engine/Source/Runtime/MediaAssets/Private/Assets/MediaPlayer.cpp
--- a/Engine/Source/Runtime/MediaAssets/Private/Assets/MediaPlayer.cpp
+++ b/Engine/Source/Runtime/MediaAssets/Private/Assets/MediaPlayer.cpp
@@ -14,6 +14,7 @@
 #include "LatentActions.h"
 #include "MediaPlayerFacade.h"
 #include "MediaPlayerOptions.h"
+#include "MediaHelpers.h"
 #include "Misc/App.h"
 #include "Misc/Paths.h"
 #include "Misc/ScopeLock.h"
@@ -365,6 +366,11 @@ bool UMediaPlayer::IsPreparing() const
 }
 
 
+bool UMediaPlayer::IsClosed() const
+{
+	return PlayerFacade->IsClosed();
+}
+
 bool UMediaPlayer::IsReady() const
 {
 	UE_LOG(LogMediaAssets, VeryVerbose, TEXT("%s.IsReady"), *GetFName().ToString());
@@ -933,6 +939,10 @@ public:
 	float TimeRemaining;
 	bool& OutSuccess;
 	bool bSawError;
+	bool bSawMediaOpened;
+	bool bSawMediaOpenFailed;
+	bool bSawSeekCompleted;
+	FMediaPlayerOptions Options;
 	FString URL;
 
 	FLatentOpenMediaSourceAction(const FLatentActionInfo& LatentInfo, UMediaPlayer* InMediaPlayer, UMediaSource* InMediaSource, const FMediaPlayerOptions& InOptions, bool& InSuccess)
@@ -943,55 +953,151 @@ public:
 		, TimeRemaining(10.0)
 		, OutSuccess(InSuccess)
 		, bSawError(false)
+		, bSawMediaOpened(false)
+		, bSawMediaOpenFailed(false)
+		, bSawSeekCompleted(false)
+		, Options(InOptions)
 	{
 		if (InMediaSource)
 		{
 			URL = InMediaSource->GetUrl();
-		}
 
-		if (!InMediaPlayer->OpenSourceWithOptions(InMediaSource, InOptions))
+			InMediaPlayer->OnMediaEvent().AddRaw(this, &FLatentOpenMediaSourceAction::HandleMediaPlayerEvent);
+			if (!InMediaPlayer->OpenSourceWithOptions(InMediaSource, InOptions))
+			{
+				UE_LOG(LogMediaAssets, Warning, TEXT("Open Media Latent: Failed initial open: %s"), *URL);
+				bSawError = true;
+			}
+		}
+		else
 		{
-			UE_LOG(LogMediaAssets, Verbose, TEXT("Latent Media: After initial open started"));
+			UE_LOG(LogMediaAssets, Warning, TEXT("Open Media Latent: Failed initial open because no media source given"));
 			bSawError = true;
 		}
 	}
 
+	virtual ~FLatentOpenMediaSourceAction()
+	{
+		UnregisterMediaEvent();
+	}
+
+	void UnregisterMediaEvent()
+	{
+		if (MediaPlayer.IsValid())
+		{
+			MediaPlayer->OnMediaEvent().RemoveAll(this);
+		}
+	}
+
+	void HandleMediaPlayerEvent(EMediaEvent Event)
+	{
+		UE_LOG(LogMediaAssets, Verbose, TEXT("Open Media Latent: Saw event: %s"), *MediaUtils::EventToString(Event));
+
+		switch (Event)
+		{
+		case EMediaEvent::MediaOpened:
+			bSawMediaOpened = true;
+			break;
+		case EMediaEvent::MediaOpenFailed:
+			bSawMediaOpenFailed = true;
+			break;
+		case EMediaEvent::SeekCompleted:
+			bSawSeekCompleted = true;
+			break;
+		}
+	}
+
+	void FailedOperation(FLatentResponse& Response)
+	{
+		OutSuccess = false;
+		Response.FinishAndTriggerIf(true, ExecutionFunction, OutputLink, CallbackTarget);
+	}
+
 	virtual void UpdateOperation(FLatentResponse& Response) override
 	{
-		if (bSawError || !MediaPlayer.IsValid() || MediaPlayer.IsStale() || MediaPlayer->HasError())
+		if (bSawMediaOpenFailed)
 		{
-			UE_LOG(LogMediaAssets, Warning, TEXT("Latent Media: Deleted or Error"));
-			OutSuccess = false;
-			Response.FinishAndTriggerIf(true, ExecutionFunction, OutputLink, CallbackTarget);
+			UE_LOG(LogMediaAssets, Warning, TEXT("Open Media Latent: Saw media open failed event. %s"), *URL);
+			FailedOperation(Response);
 			return;
 		}
 
-		if (MediaPlayer->IsReady())
+		if (!MediaPlayer.IsValid() || MediaPlayer.IsStale())
 		{
-			UE_LOG(LogMediaAssets, Verbose, TEXT("Latent Media: IsReady"));
-			OutSuccess = true;
-			Response.FinishAndTriggerIf(true, ExecutionFunction, OutputLink, CallbackTarget);
+			UE_LOG(LogMediaAssets, Warning, TEXT("Open Media Latent: Media player object was deleted. %s"), *URL);
+			FailedOperation(Response);
 			return;
 		}
-		else if (!MediaPlayer->IsPreparing())
+
+		if (bSawError || MediaPlayer->HasError())
 		{
-			UE_LOG(LogMediaAssets, Verbose, TEXT("Latent Media: Invalid State"));
-			OutSuccess = false;
-			Response.FinishAndTriggerIf(true, ExecutionFunction, OutputLink, CallbackTarget);
+			UE_LOG(LogMediaAssets, Warning, TEXT("Open Media Latent: Media player is in Error state. %s"), *URL);
+			FailedOperation(Response);
 			return;
 		}
+
+		if (MediaPlayer->IsClosed())
+		{
+			UE_LOG(LogMediaAssets, Warning, TEXT("Open Media Latent: Media player is closed. %s"), *URL);
+			FailedOperation(Response);
+			return;
+		}
+
+		if (MediaPlayer->IsPreparing())
+		{
+			UE_LOG(LogMediaAssets, Verbose, TEXT("Open Media Latent: Is preparing ..."));
+		}
+		else if (MediaPlayer->IsReady())
+		{
+			UE_LOG(LogMediaAssets, Verbose, TEXT("Open Media Latent: IsReady() ... %s"), *URL);
+
+			if (bSawMediaOpened)
+			{
+				if (!Options.SeekTime.IsZero())
+				{
+					if (bSawSeekCompleted)
+					{
+						UE_LOG(LogMediaAssets, Verbose, TEXT("Open Media Latent: Triggering output pin after seek completed. Success: %d, %s"), OutSuccess, *URL);
+						OutSuccess = true;
+						Response.FinishAndTriggerIf(true, ExecutionFunction, OutputLink, CallbackTarget);
+						return;
+					}
+					else
+					{
+						if (Options.SeekTime < FTimespan::FromSeconds(0) || Options.SeekTime > MediaPlayer->GetDuration())
+						{
+							UE_LOG(LogMediaAssets, Warning, TEXT("Open Media Latent: Media player seeking to time out of bounds. Seek: %s, Duration: %s, URL: %s"), 
+								*Options.SeekTime.ToString(), *MediaPlayer->GetDuration().ToString(), *URL);
+							FailedOperation(Response);
+							return;
+						}
+						UE_LOG(LogMediaAssets, Verbose, TEXT("Open Media Latent: Waiting for seek completed event ..."));
+					}
+				}
+				else
+				{
+					UE_LOG(LogMediaAssets, Verbose, TEXT("Open Media Latent: Triggering output pin after opened event. Success: %d, %s"), OutSuccess, *URL);
+					OutSuccess = true;
+					Response.FinishAndTriggerIf(true, ExecutionFunction, OutputLink, CallbackTarget);
+					return;
+				}
+			}
+			else
+			{
+				UE_LOG(LogMediaAssets, Verbose, TEXT("Open Media Latent: Waiting for opened event ..."));
+			}
+		}
 		else
 		{
-			UE_LOG(LogMediaAssets, Verbose, TEXT("Latent Media: Waiting for open ..."));
+			UE_LOG(LogMediaAssets, Verbose, TEXT("Open Media Latent: Waiting for IsReady() ..."));
 		}
 
 		// Timed out
 		TimeRemaining -= Response.ElapsedTime();
 		if (TimeRemaining <= 0.0f)
 		{
-			UE_LOG(LogMediaAssets, Verbose, TEXT("Latent Media: Timed out."));
-			OutSuccess = false;
-			Response.FinishAndTriggerIf(true, ExecutionFunction, OutputLink, CallbackTarget);
+			UE_LOG(LogMediaAssets, Warning, TEXT("Open Media Latent: Timed out. %s"), *URL);
+			FailedOperation(Response);
 			return;
 		}
 	}
diff --git a/Engine/Source/Runtime/MediaAssets/Public/MediaPlayer.h b/Engine/Source/Runtime/MediaAssets/Public/MediaPlayer.h
--- a/Engine/Source/Runtime/MediaAssets/Public/MediaPlayer.h
+++ b/Engine/Source/Runtime/MediaAssets/Public/MediaPlayer.h
@@ -492,6 +492,14 @@ public:
 	UFUNCTION(BlueprintCallable, Category="Media|MediaPlayer")
 	bool IsPreparing() const;
 
+	/**
+	 * Whether media is currently closed.
+	 *
+	 * @return true if media is closed, false otherwise.
+	 */
+	UFUNCTION(BlueprintCallable, Category="Media|MediaPlayer")
+	bool IsClosed() const;
+
 	/**
 	 * Check whether media is ready for playback.
 	 *
@@ -589,11 +597,11 @@ public:
 	/**
 	 * Open the specified media source with options using a latent action.
 	 *
-	 * A result of true indicates that the player opened the media source successfully.
+	 * A result of true indicates that the player successfully completed all requested operations.
 	 *
 	 * @param MediaSource The media source to open.
 	 * @param Options The media player options to apply.
-	 * @param bSuccess  Player opened the media.
+	 * @param bSuccess  All requested operations have completed successfully.
 	 * @see Close, OpenFile, OpenPlaylist, OpenPlaylistIndex, OpenUrl, Reopen
 	 */
 	UFUNCTION(BlueprintCallable, Category = "Media|MediaPlayer", meta = (Latent, LatentInfo = "LatentInfo", WorldContext = "WorldContextObject"))
diff --git a/Engine/Source/Runtime/MediaUtils/Private/MediaPlayerFacade.cpp b/Engine/Source/Runtime/MediaUtils/Private/MediaPlayerFacade.cpp
--- a/Engine/Source/Runtime/MediaUtils/Private/MediaPlayerFacade.cpp
+++ b/Engine/Source/Runtime/MediaUtils/Private/MediaPlayerFacade.cpp
@@ -412,6 +412,10 @@ bool FMediaPlayerFacade::IsPreparing() const
 	return Player.IsValid() && (Player->GetControls().GetState() == EMediaState::Preparing);
 }
 
+bool FMediaPlayerFacade::IsClosed() const
+{
+	return Player.IsValid() && (Player->GetControls().GetState() == EMediaState::Closed);
+}
 
 bool FMediaPlayerFacade::IsReady() const
 {
diff --git a/Engine/Source/Runtime/MediaUtils/Public/MediaPlayerFacade.h b/Engine/Source/Runtime/MediaUtils/Public/MediaPlayerFacade.h
--- a/Engine/Source/Runtime/MediaUtils/Public/MediaPlayerFacade.h
+++ b/Engine/Source/Runtime/MediaUtils/Public/MediaPlayerFacade.h
@@ -472,6 +472,13 @@ public:
 	 */
 	bool IsPreparing() const;
 
+	/**
+	 * Whether media is currently closed.
+	 *
+	 * @return true if media is closed, false otherwise.
+	 */
+	bool IsClosed() const;
+
 	/**
 	 * Whether media is ready for playback.
 	 *
diff --git a/Engine/Source/Runtime/RenderCore/Private/VertexFactory.cpp b/Engine/Source/Runtime/RenderCore/Private/VertexFactory.cpp
--- a/Engine/Source/Runtime/RenderCore/Private/VertexFactory.cpp
+++ b/Engine/Source/Runtime/RenderCore/Private/VertexFactory.cpp
@@ -210,8 +210,7 @@ void FVertexFactory::GetStreams(ERHIFeatureLevel::Type InFeatureLevel, FVertexIn
 
 void FVertexFactory::OffsetInstanceStreams(uint32 InstanceOffset, bool bOperateOnPositionOnly, FVertexInputStreamArray& VertexStreams) const
 {
-	const TArray<FVertexStream,TFixedAllocator<MaxVertexElementCount> >& StreamArray = bOperateOnPositionOnly ? PositionStream : Streams;
-
+	const TArrayView<const FVertexStream>& StreamArray = bOperateOnPositionOnly ? MakeArrayView(PositionStream) : MakeArrayView(Streams);
 	for (int32 StreamIndex = 0; StreamIndex < StreamArray.Num(); StreamIndex++)
 	{
 		const FVertexStream& Stream = StreamArray[StreamIndex];
diff --git a/Engine/Source/Runtime/Renderer/Private/SceneVisibility.cpp b/Engine/Source/Runtime/Renderer/Private/SceneVisibility.cpp
--- a/Engine/Source/Runtime/Renderer/Private/SceneVisibility.cpp
+++ b/Engine/Source/Runtime/Renderer/Private/SceneVisibility.cpp
@@ -2048,7 +2048,7 @@ struct FRelevancePacket
 
 			if (OutHasViewCustomDataMasks[PrimitiveIndex] != 0) // Has a relevance for this view
 			{
-				UserViewCustomData = PrimitiveSceneInfo->Proxy->InitViewCustomData(View, View.LODDistanceFactor, PrimitiveCustomDataMemStack, true, &LODToRender, MeshScreenSizeSquared);
+				UserViewCustomData = PrimitiveSceneInfo->Proxy->InitViewCustomData(View, View.LODDistanceFactor, PrimitiveCustomDataMemStack, true, false, &LODToRender, MeshScreenSizeSquared);
 
 				if (UserViewCustomData != nullptr)
 				{
@@ -2528,7 +2528,7 @@ static void SetDynamicMeshElementViewCustomData(TArray<FViewInfo>& InViews, cons
 
 			if (InHasViewCustomDataMasks[PrimitiveIndex] & (1 << ViewIndex) && ViewInfo.GetCustomData(InPrimitiveSceneInfo->GetIndex()) == nullptr)
 			{
-				ViewInfo.SetCustomData(InPrimitiveSceneInfo, InPrimitiveSceneInfo->Proxy->InitViewCustomData(ViewInfo, ViewInfo.LODDistanceFactor, ViewInfo.GetCustomDataGlobalMemStack()));
+				ViewInfo.SetCustomData(InPrimitiveSceneInfo, InPrimitiveSceneInfo->Proxy->InitViewCustomData(ViewInfo, ViewInfo.LODDistanceFactor, ViewInfo.GetCustomDataGlobalMemStack(), false, false));
 			}
 		}
 	}
diff --git a/Engine/Source/Runtime/Renderer/Private/ShadowSetup.cpp b/Engine/Source/Runtime/Renderer/Private/ShadowSetup.cpp
--- a/Engine/Source/Runtime/Renderer/Private/ShadowSetup.cpp
+++ b/Engine/Source/Runtime/Renderer/Private/ShadowSetup.cpp
@@ -877,7 +877,7 @@ bool FProjectedShadowInfo::ShouldDrawStaticMeshes(FViewInfo& InCurrentView, bool
 			{
 				if (InCurrentView.GetCustomData(InPrimitiveSceneInfo->GetIndex()) == nullptr)
 				{
-					InCurrentView.SetCustomData(InPrimitiveSceneInfo, InPrimitiveSceneInfo->Proxy->InitViewCustomData(InCurrentView, InCurrentView.LODDistanceFactor, InCurrentView.GetCustomDataGlobalMemStack(), true, &ViewLODToRender, MeshScreenSizeSquared));
+					InCurrentView.SetCustomData(InPrimitiveSceneInfo, InPrimitiveSceneInfo->Proxy->InitViewCustomData(InCurrentView, InCurrentView.LODDistanceFactor, InCurrentView.GetCustomDataGlobalMemStack(), true, true, &ViewLODToRender, MeshScreenSizeSquared));
 				}
 			}
 		}
@@ -1472,7 +1472,7 @@ void FProjectedShadowInfo::GatherDynamicMeshElementsArray(
 			{
 				if (DependentView->GetCustomData(PrimitiveSceneInfo->GetIndex()) == nullptr)
 				{
-					void* CustomData = PrimitiveSceneInfo->Proxy->InitViewCustomData(*DependentView, DependentView->LODDistanceFactor, DependentView->GetCustomDataGlobalMemStack());
+					void* CustomData = PrimitiveSceneInfo->Proxy->InitViewCustomData(*DependentView, DependentView->LODDistanceFactor, DependentView->GetCustomDataGlobalMemStack(), false, true);
 					DependentView->SetCustomData(PrimitiveSceneInfo, CustomData);
 
 					// This is required as GetDynamicMeshElements will received ReusedViewsArray which contains only FoundView
