Fixed Switch media decoder track selection behavior, end of stream detection, random failure of duration computation, limited output queue depth to ease memory pressure & switched audio output to use more appropriate "unsynced" mode (and edits still function with dev-main code)

diff --git a/Engine/Source/Runtime/Media/Public/IMediaPlayer.h b/Engine/Source/Runtime/Media/Public/IMediaPlayer.h
--- a/Engine/Source/Runtime/Media/Public/IMediaPlayer.h
+++ b/Engine/Source/Runtime/Media/Public/IMediaPlayer.h
@@ -175,7 +175,7 @@ public:
 	*/
 	virtual bool RequiresAudioSyncAfterDropouts() const
 	{
-#if PLATFORM_PS4 || PLATFORM_SWITCH || PLATFORM_XBOXONE
+#if PLATFORM_PS4 || PLATFORM_XBOXONE
 		return true;
 #else
 		return false;
diff --git a/Engine/Source/Runtime/MediaAssets/Private/Assets/MediaSoundComponent.cpp b/Engine/Source/Runtime/MediaAssets/Private/Assets/MediaSoundComponent.cpp
--- a/Engine/Source/Runtime/MediaAssets/Private/Assets/MediaSoundComponent.cpp
+++ b/Engine/Source/Runtime/MediaAssets/Private/Assets/MediaSoundComponent.cpp
@@ -72,7 +72,7 @@ UMediaSoundComponent::UMediaSoundComponent(const FObjectInitializer& ObjectIniti
 	bVisualizeComponent = true;
 #endif
 
-#if PLATFORM_PS4 || PLATFORM_SWITCH || PLATFORM_XBOXONE
+#if PLATFORM_PS4 || PLATFORM_XBOXONE
 	bSyncAudioAfterDropouts = true;
 #else
 	bSyncAudioAfterDropouts = false;
@@ -422,13 +422,18 @@ int32 UMediaSoundComponent::OnGenerateAudio(float* OutAudio, int32 NumSamples)
 		}
 		else
 		{
-			const int32 FramesRequested = NumSamples / NumChannels;
+			const uint32 FramesRequested = uint32(NumSamples / NumChannels);
 			uint32 JumpFrame = MAX_uint32;
-			uint32 FramesWritten = Resampler->Generate(OutAudio, OutTime, (uint32)FramesRequested, Rate, Time, *PinnedSampleQueue, JumpFrame);
+			uint32 FramesWritten = Resampler->Generate(OutAudio, OutTime, FramesRequested, Rate, Time, *PinnedSampleQueue, JumpFrame);
 			if (FramesWritten == 0)
 			{
 				return 0; // no samples available
 			}
+
+			if (FramesWritten < FramesRequested)
+			{
+				memset(OutAudio + FramesWritten * NumChannels, 0, (NumSamples - FramesWritten * NumChannels) * sizeof(float));
+			}
 		}
 
 		LastPlaySampleTime = OutTime;
