Add 0 entries to enums to prevent blueprint compiler warnings.

Enums affected:
- EVoiceSampleRate
- EMediaWebcamCaptureDeviceFilter
- EMediaVideoCaptureDeviceFilter
- EMediaAudioCaptureDeviceFilter

#rnx

diff --git a/Engine/Source/Runtime/MediaAssets/Public/Misc/MediaBlueprintFunctionLibrary.h b/Engine/Source/Runtime/MediaAssets/Public/Misc/MediaBlueprintFunctionLibrary.h
--- a/Engine/Source/Runtime/MediaAssets/Public/Misc/MediaBlueprintFunctionLibrary.h
+++ b/Engine/Source/Runtime/MediaAssets/Public/Misc/MediaBlueprintFunctionLibrary.h
@@ -15,6 +15,8 @@
 UENUM(BlueprintType, meta=(BitFlags, UseEnumValuesAsMaskValuesInEditor = "true"))
 enum class EMediaAudioCaptureDeviceFilter : uint8
 {
+	None = 0 UMETA(Hidden),
+
 	/** Audio capture cards. */
 	Card = 0x1,
 
@@ -37,6 +39,8 @@ ENUM_CLASS_FLAGS(EMediaAudioCaptureDeviceFilter)
 UENUM(BlueprintType, meta=(BitFlags, UseEnumValuesAsMaskValuesInEditor = "true"))
 enum class EMediaVideoCaptureDeviceFilter : uint8
 {
+	None = 0 UMETA(Hidden),
+
 	/** Video capture card. */
 	Card = 0x1,
 
@@ -59,6 +63,8 @@ ENUM_CLASS_FLAGS(EMediaVideoCaptureDeviceFilter)
 UENUM(BlueprintType, meta=(BitFlags, UseEnumValuesAsMaskValuesInEditor = "true"))
 enum class EMediaWebcamCaptureDeviceFilter : uint8
 {
+	None = 0 UMETA(Hidden),
+
 	/** Depth sensor. */
 	DepthSensor = 0x1,
 
