MediaPlayer: Events are hidden in blueprints editor details panel as they should not be used there.

diff --git a/Engine/Source/Editor/Kismet/Private/BlueprintDetailsCustomization.cpp b/Engine/Source/Editor/Kismet/Private/BlueprintDetailsCustomization.cpp
--- a/Engine/Source/Editor/Kismet/Private/BlueprintDetailsCustomization.cpp
+++ b/Engine/Source/Editor/Kismet/Private/BlueprintDetailsCustomization.cpp
@@ -107,9 +107,10 @@ void FBlueprintDetails::AddEventsCategory(IDetailLayoutBuilder& DetailBuilder, U
 					UProperty* Property = *PropertyIt;
 
 					FName PropertyName = ComponentProperty->GetFName();
-
+					static const FName HideInDetailPanelName("HideInDetailPanel");
 					// Check for multicast delegates that we can safely assign
-					if ( !Property->HasAnyPropertyFlags(CPF_Parm) && Property->HasAllPropertyFlags(CPF_BlueprintAssignable) )
+					if ( !Property->HasAnyPropertyFlags(CPF_Parm) && Property->HasAllPropertyFlags(CPF_BlueprintAssignable) &&
+						!Property->HasMetaData(HideInDetailPanelName) )
 					{
 						FName EventName = Property->GetFName();
 						FText EventText = Property->GetDisplayNameText();
diff --git a/Engine/Source/Runtime/MediaAssets/Public/MediaPlayer.h b/Engine/Source/Runtime/MediaAssets/Public/MediaPlayer.h
--- a/Engine/Source/Runtime/MediaAssets/Public/MediaPlayer.h
+++ b/Engine/Source/Runtime/MediaAssets/Public/MediaPlayer.h
@@ -849,11 +849,11 @@ public:
 public:
 
 	/** A delegate that is invoked when playback has reached the end of the media. */
-	UPROPERTY(BlueprintAssignable, Category="Media|MediaPlayer")
+	UPROPERTY(BlueprintAssignable, Category="Media|MediaPlayer", meta = (HideInDetailPanel))
 	FOnMediaPlayerMediaEvent OnEndReached;
 
 	/** A delegate that is invoked when a media source has been closed. */
-	UPROPERTY(BlueprintAssignable, Category="Media|MediaPlayer")
+	UPROPERTY(BlueprintAssignable, Category="Media|MediaPlayer", meta = (HideInDetailPanel))
 	FOnMediaPlayerMediaEvent OnMediaClosed;
 
 	/**
@@ -865,7 +865,7 @@ public:
 	 *
 	 * @see OnMediaOpenFailed, OnTracksChanged
 	 */
-	UPROPERTY(BlueprintAssignable, Category="Media|MediaPlayer")
+	UPROPERTY(BlueprintAssignable, Category="Media|MediaPlayer", meta = (HideInDetailPanel))
 	FOnMediaPlayerMediaOpened OnMediaOpened;
 
 	/**
@@ -877,7 +877,7 @@ public:
 	 *
 	 * @see OnMediaOpened
 	 */
-	UPROPERTY(BlueprintAssignable, Category="Media|MediaPlayer")
+	UPROPERTY(BlueprintAssignable, Category="Media|MediaPlayer", meta = (HideInDetailPanel))
 	FOnMediaPlayerMediaOpenFailed OnMediaOpenFailed;
 
 	/**
@@ -885,7 +885,7 @@ public:
 	 *
 	 * @see OnPlaybackSuspended
 	 */
-	UPROPERTY(BlueprintAssignable, Category="Media|MediaPlayer")
+	UPROPERTY(BlueprintAssignable, Category="Media|MediaPlayer", meta = (HideInDetailPanel))
 	FOnMediaPlayerMediaEvent OnPlaybackResumed;
 
 	/**
@@ -893,7 +893,7 @@ public:
 	 *
 	 * @see OnPlaybackResumed
 	 */
-	UPROPERTY(BlueprintAssignable, Category="Media|MediaPlayer")
+	UPROPERTY(BlueprintAssignable, Category="Media|MediaPlayer", meta = (HideInDetailPanel))
 	FOnMediaPlayerMediaEvent OnPlaybackSuspended;
 
 	/**
@@ -903,7 +903,7 @@ public:
 	 * synchronously or asynchronously, this event may be executed before or
 	 * after the call to Seek returns.
 	 */
-	UPROPERTY(BlueprintAssignable, Category="Media|MediaPlayer")
+	UPROPERTY(BlueprintAssignable, Category="Media|MediaPlayer", meta = (HideInDetailPanel))
 	FOnMediaPlayerMediaEvent OnSeekCompleted;
 
 	/**
@@ -911,7 +911,7 @@ public:
 	 *
 	 * @see OnMediaOpened
 	 */
-	UPROPERTY(BlueprintAssignable, Category="Media|MediaPlayer")
+	UPROPERTY(BlueprintAssignable, Category="Media|MediaPlayer", meta = (HideInDetailPanel))
 	FOnMediaPlayerMediaEvent OnTracksChanged;
 
 public:
