Fixing compile error: missing dependency of MediaSoundComponent on SignalProcesing
#rnx

diff --git a/Engine/Source/Runtime/MediaAssets/MediaAssets.Build.cs b/Engine/Source/Runtime/MediaAssets/MediaAssets.Build.cs
--- a/Engine/Source/Runtime/MediaAssets/MediaAssets.Build.cs
+++ b/Engine/Source/Runtime/MediaAssets/MediaAssets.Build.cs
@@ -8,8 +8,9 @@ namespace UnrealBuildTool.Rules
 		{
 			PublicDependencyModuleNames.AddRange(
 				new string[] {
+					"SignalProcessing",
 					"AudioMixer",
-					"Core",
+                    "Core",
 					"CoreUObject",
 					"Engine",
 					"Media",
