Fixing up "AutoClear" feature for MediaTexture object for various cases / platforms
#rnx

diff --git a/Engine/Source/Runtime/MediaAssets/Private/Misc/MediaTextureResource.cpp b/Engine/Source/Runtime/MediaAssets/Private/Misc/MediaTextureResource.cpp
--- a/Engine/Source/Runtime/MediaAssets/Private/Misc/MediaTextureResource.cpp
+++ b/Engine/Source/Runtime/MediaAssets/Private/Misc/MediaTextureResource.cpp
@@ -250,6 +250,7 @@ void FMediaTextureResource::Render(const FRenderParams& Params)
 			{
 				CreateOutputRenderTarget(Sample, Params);
 				Sample->GetMediaTextureSampleConverter()->Convert(RenderTargetTextureRHI);
+				Cleared = false;
 			}
 			else if (MediaTextureResource::RequiresConversion(Sample, Params.SrgbOutput))
 			{
@@ -872,6 +873,8 @@ void FMediaTextureResource::CreateOutputRenderTarget(const TSharedPtr<IMediaText
 
 		CurrentClearColor = InParams.ClearColor;
 		UpdateResourceSize();
+
+		Cleared = false;
 	}
 
 	if (RenderTargetTextureRHI != OutputTarget)
