Update of Electra and MediaFramework (includes Android platform support & QoS changes)

diff --git a/Engine/Plugins/Media/WmfMedia/Source/WmfMedia/Private/Wmf/WmfMediaHardwareVideoDecodingRendering.cpp b/Engine/Plugins/Media/WmfMedia/Source/WmfMedia/Private/Wmf/WmfMediaHardwareVideoDecodingRendering.cpp
--- a/Engine/Plugins/Media/WmfMedia/Source/WmfMedia/Private/Wmf/WmfMediaHardwareVideoDecodingRendering.cpp
+++ b/Engine/Plugins/Media/WmfMedia/Source/WmfMedia/Private/Wmf/WmfMediaHardwareVideoDecodingRendering.cpp
@@ -14,6 +14,7 @@
 #include "SceneUtils.h"
 #include "SceneInterface.h"
 #include "ShaderParameterUtils.h"
+#include "ProfilingDebugging/RealtimeGPUProfiler.h"
 
 #include "D3D11RHIPrivate.h"
 #include "DynamicRHI.h"
@@ -29,7 +30,10 @@
 
 #include "WmfMediaHardwareVideoDecodingShaders.h"
 
-FRHICOMMAND_MACRO(FRHICommandCopyResource)
+DECLARE_GPU_STAT_NAMED(MediaTextureConversion, TEXT("MediaTextureConversion"));
+
+
+struct FRHICommandCopyResource final : public FRHICommand<FRHICommandCopyResource>
 {
 	TComPtr<ID3D11Texture2D> SampleTexture;
 	FTexture2DRHIRef SampleDestinationTexture;
@@ -42,7 +46,7 @@ FRHICOMMAND_MACRO(FRHICommandCopyResource)
 
 	void Execute(FRHICommandListBase& CmdList)
 	{
-		LLM_SCOPE(ELLMTag::VideoStreaming);
+		LLM_SCOPE(ELLMTag::MediaStreaming);
 		ID3D11Device* D3D11Device = static_cast<ID3D11Device*>(GDynamicRHI->RHIGetNativeDevice());
 		ID3D11DeviceContext* D3D11DeviceContext = nullptr;
 
@@ -104,7 +108,8 @@ FRHICOMMAND_MACRO(FRHICommandCopyResource)
 
 bool FWmfMediaHardwareVideoDecodingParameters::ConvertTextureFormat_RenderThread(FWmfMediaHardwareVideoDecodingTextureSample* InSample, FTexture2DRHIRef InDstTexture)
 {
-	LLM_SCOPE(ELLMTag::VideoStreaming);
+	LLM_SCOPE(ELLMTag::MediaStreaming);
+
 	if (InSample == nullptr || !InDstTexture.IsValid())
 	{
 		return false;
@@ -124,6 +129,9 @@ bool FWmfMediaHardwareVideoDecodingParameters::ConvertTextureFormat_RenderThread
 	{
 		FRHICommandListImmediate& RHICmdList = FRHICommandListExecutor::GetImmediateCommandList();
 
+		SCOPED_DRAW_EVENT(RHICmdList, FWmfMediaHardwareVideoDecodingParameters_Convert);
+		SCOPED_GPU_STAT(RHICmdList, MediaTextureConversion);
+
 		FRHIRenderPassInfo RPInfo(InDstTexture, ERenderTargetActions::DontLoad_Store);
 		RHICmdList.BeginRenderPass(RPInfo, TEXT("ConvertTextureFormat"));
 
diff --git a/Engine/Source/Runtime/Core/Private/HAL/LowLevelMemTracker.cpp b/Engine/Source/Runtime/Core/Private/HAL/LowLevelMemTracker.cpp
--- a/Engine/Source/Runtime/Core/Private/HAL/LowLevelMemTracker.cpp
+++ b/Engine/Source/Runtime/Core/Private/HAL/LowLevelMemTracker.cpp
@@ -106,7 +106,9 @@ DECLARE_LLM_MEMORY_STAT(TEXT("MaterialInstance"), STAT_MaterialInstanceLLM, STAT
 DECLARE_LLM_MEMORY_STAT(TEXT("SkeletalMesh"), STAT_SkeletalMeshLLM, STATGROUP_LLMFULL);
 DECLARE_LLM_MEMORY_STAT(TEXT("InstancedMesh"), STAT_InstancedMeshLLM, STATGROUP_LLMFULL);
 DECLARE_LLM_MEMORY_STAT(TEXT("Landscape"), STAT_LandscapeLLM, STATGROUP_LLMFULL);
-DECLARE_LLM_MEMORY_STAT(TEXT("VideoStreaming"), STAT_VideoStreamingLLM, STATGROUP_LLMFULL);
+DECLARE_LLM_MEMORY_STAT(TEXT("MediaStreaming"), STAT_MediaStreamingLLM, STATGROUP_LLMFULL);
+DECLARE_LLM_MEMORY_STAT(TEXT("ElectraPlayer"), STAT_ElectraPlayerLLM, STATGROUP_LLMFULL);
+DECLARE_LLM_MEMORY_STAT(TEXT("WMFPlayer"), STAT_WMFPlayerLLM, STATGROUP_LLMFULL);
 DECLARE_LLM_MEMORY_STAT(TEXT("MMIO"), STAT_PlatformMMIOLLM, STATGROUP_LLMPlatform);
 DECLARE_LLM_MEMORY_STAT(TEXT("VirtualMemory"), STAT_PlatformVMLLM, STATGROUP_LLMPlatform);
 
@@ -127,6 +129,7 @@ DECLARE_LLM_MEMORY_STAT(TEXT("Particles"), STAT_ParticlesSummaryLLM, STATGROUP_L
 DECLARE_LLM_MEMORY_STAT(TEXT("Niagara"), STAT_NiagaraSummaryLLM, STATGROUP_LLM);
 DECLARE_LLM_MEMORY_STAT(TEXT("UI"), STAT_UISummaryLLM, STATGROUP_LLM);
 DECLARE_LLM_MEMORY_STAT(TEXT("Textures"), STAT_TexturesSummaryLLM, STATGROUP_LLM);
+DECLARE_LLM_MEMORY_STAT(TEXT("MediaStreaming"), STAT_MediaStreamingSummaryLLM, STATGROUP_LLM);
 
 extern const TCHAR* LLMGetTagName(ELLMTag Tag)
 {
diff --git a/Engine/Source/Runtime/Core/Public/HAL/LowLevelMemTracker.h b/Engine/Source/Runtime/Core/Public/HAL/LowLevelMemTracker.h
--- a/Engine/Source/Runtime/Core/Public/HAL/LowLevelMemTracker.h
+++ b/Engine/Source/Runtime/Core/Public/HAL/LowLevelMemTracker.h
@@ -205,7 +205,9 @@ enum class ELLMTagSet : uint8
 	macro(InstancedMesh,						"InstancedMesh",				GET_STATFNAME(STAT_InstancedMeshLLM),						GET_STATFNAME(STAT_EngineSummaryLLM),			ELLMTag::Meshes)\
 	macro(Landscape,							"Landscape",					GET_STATFNAME(STAT_LandscapeLLM),							GET_STATFNAME(STAT_EngineSummaryLLM),			ELLMTag::Meshes)\
 	macro(CsvProfiler,							"CsvProfiler",					GET_STATFNAME(STAT_CsvProfilerLLM),							GET_STATFNAME(STAT_EngineSummaryLLM),			-1)\
-	macro(VideoStreaming,						"VideoStreaming",				GET_STATFNAME(STAT_VideoStreamingLLM),						GET_STATFNAME(STAT_EngineSummaryLLM),			-1)\
+	macro(MediaStreaming,						"MediaStreaming",				GET_STATFNAME(STAT_MediaStreamingLLM),						GET_STATFNAME(STAT_MediaStreamingSummaryLLM),	-1)\
+	macro(ElectraPlayer,						"ElectraPlayer",				GET_STATFNAME(STAT_ElectraPlayerLLM),						GET_STATFNAME(STAT_MediaStreamingSummaryLLM),	ELLMTag::MediaStreaming)\
+	macro(WMFPlayer,							"WMFPlayer",					GET_STATFNAME(STAT_WMFPlayerLLM),							GET_STATFNAME(STAT_MediaStreamingSummaryLLM),	ELLMTag::MediaStreaming)\
 	macro(PlatformMMIO,							"MMIO",							GET_STATFNAME(STAT_PlatformMMIOLLM),						NAME_None,										-1)\
 	macro(PlatformVM,							"Virtual Memory",				GET_STATFNAME(STAT_PlatformVMLLM),							NAME_None,										-1)\
 
diff --git a/Engine/Source/Runtime/Media/Private/MediaModule.cpp b/Engine/Source/Runtime/Media/Private/MediaModule.cpp
--- a/Engine/Source/Runtime/Media/Private/MediaModule.cpp
+++ b/Engine/Source/Runtime/Media/Private/MediaModule.cpp
@@ -7,6 +7,7 @@
 #include "Misc/ScopeLock.h"
 #include "Modules/ModuleManager.h"
 #include "Stats/Stats.h"
+#include "ProfilingDebugging/CsvProfiler.h"
 
 #include "IMediaCaptureSupport.h"
 #include "IMediaPlayerFactory.h"
@@ -14,6 +15,7 @@
 #include "MediaClock.h"
 #include "MediaTicker.h"
 
+CSV_DEFINE_CATEGORY_MODULE(MEDIA_API, MediaStreaming, false);
 
 /**
  * Implements the Media module.
diff --git a/Engine/Source/Runtime/MediaAssets/Private/Assets/MediaSoundComponent.cpp b/Engine/Source/Runtime/MediaAssets/Private/Assets/MediaSoundComponent.cpp
--- a/Engine/Source/Runtime/MediaAssets/Private/Assets/MediaSoundComponent.cpp
+++ b/Engine/Source/Runtime/MediaAssets/Private/Assets/MediaSoundComponent.cpp
@@ -11,6 +11,7 @@
 #include "Misc/ScopeLock.h"
 #include "Sound/AudioSettings.h"
 #include "UObject/UObjectGlobals.h"
+#include "ProfilingDebugging/CsvProfiler.h"
 
 #include "MediaPlayer.h"
 #include "MediaPlayerFacade.h"
@@ -28,6 +29,8 @@ DECLARE_FLOAT_COUNTER_STAT(TEXT("MediaUtils MediaSoundComponent Sync"), STAT_Med
 DECLARE_FLOAT_COUNTER_STAT(TEXT("MediaUtils MediaSoundComponent SampleTime"), STAT_MediaUtils_MediaSoundComponentSampleTime, STATGROUP_Media);
 DECLARE_DWORD_COUNTER_STAT(TEXT("MediaUtils MediaSoundComponent Queued"), STAT_Media_SoundCompQueued, STATGROUP_Media);
 
+CSV_DECLARE_CATEGORY_MODULE_EXTERN(MEDIA_API, MediaStreaming);
+
 /* Static initialization
  *****************************************************************************/
 
@@ -337,6 +340,8 @@ bool UMediaSoundComponent::Init(int32& SampleRate)
 
 int32 UMediaSoundComponent::OnGenerateAudio(float* OutAudio, int32 NumSamples)
 {
+	CSV_SCOPED_TIMING_STAT(MediaStreaming, UMediaSoundComponent_OnGenerateAudio);
+
 	int32 InitialSyncOffset = 0;
 	TSharedPtr<FMediaAudioSampleQueue, ESPMode::ThreadSafe> PinnedSampleQueue;
 	{
diff --git a/Engine/Source/Runtime/MediaAssets/Private/Misc/MediaTextureResource.cpp b/Engine/Source/Runtime/MediaAssets/Private/Misc/MediaTextureResource.cpp
--- a/Engine/Source/Runtime/MediaAssets/Private/Misc/MediaTextureResource.cpp
+++ b/Engine/Source/Runtime/MediaAssets/Private/Misc/MediaTextureResource.cpp
@@ -20,6 +20,7 @@
 #include "StaticBoundShaderState.h"
 #include "RenderUtils.h"
 #include "RHIStaticStates.h"
+#include "ProfilingDebugging/CsvProfiler.h"
 
 #include "MediaTexture.h"
 
@@ -33,6 +34,10 @@ DECLARE_CYCLE_STAT(TEXT("MediaAssets MediaTextureResource Render"), STAT_MediaAs
 /** Sample time of texture last rendered. */
 DECLARE_FLOAT_COUNTER_STAT(TEXT("MediaAssets MediaTextureResource Sample"), STAT_MediaUtils_TextureSampleTime, STATGROUP_Media);
 
+CSV_DECLARE_CATEGORY_MODULE_EXTERN(MEDIA_API, MediaStreaming);
+
+DECLARE_GPU_STAT_NAMED(MediaTextureResource, TEXT("MediaTextureResource"));
+
 /* Local helpers
  *****************************************************************************/
 
@@ -177,7 +182,9 @@ void FMediaTextureResource::Render(const FRenderParams& Params)
 {
 	check(IsInRenderingThread());
 
+	LLM_SCOPE(ELLMTag::MediaStreaming);
 	SCOPE_CYCLE_COUNTER(STAT_MediaAssets_MediaTextureResourceRender);
+	CSV_SCOPED_TIMING_STAT(MediaStreaming, FMediaTextureResource_Render);
 
 	FLinearColor Rotation(1, 0, 0, 1);
 	FLinearColor Offset(0, 0, 0, 0);
@@ -280,49 +287,44 @@ void FMediaTextureResource::Render(const FRenderParams& Params)
 			Rotation = Sample->GetScaleRotation();
 			Offset = Sample->GetOffset();
 
-			/*
-				We need to keep any sample used for display around to avoid any player reusing it before any async tasks are done.
+			// Keep new current sample around to avoid it returning to any pool too soon
+			// (we hold on to rendering resources at a lower level anyway)
+			CurrentSample = TRefCountPtr<FTextureSampleKeeper>(new FTextureSampleKeeper(Sample));
+			check(CurrentSample);
 
-				Two cases are to be considered:
-				a) we directly use the textures in the sample -> we need to keep it as long as the sample is visible
-				b) we copy (e.g. as part of a conversion) the sample into a local texture -> we only need it for the current rendering interval
+			SET_FLOAT_STAT(STAT_MediaUtils_TextureSampleTime, Sample->GetTime().GetTotalMilliseconds());
+		}
+		else
+		{
+			if (OutputTarget == RenderTargetTextureRHI)
+			{
+				/*
+					We know that the current sample is using the internal buffer to store converted data.
+					This data is the data actually used by any user of this resource. Hence we can release
+					the current sample as soon as the first frame using it is done processing.
+					Either this will happen when a new sample arrives or we do it explicitly here.
+				*/
+				CurrentSample = nullptr;
+			}
 
-				In both cases the FTextureSampleKeeper will make sure that it is kept for the current rendering interval per FRHIResource rules
-				for delayed deletion.
-			*/
-			TRefCountPtr<FTextureSampleKeeper> NewCurrentSample = TRefCountPtr<FTextureSampleKeeper>(new FTextureSampleKeeper(Sample));
-			check(NewCurrentSample);
-//TODO: using normal sample retention at all times for now to avoid issues on some platforms
-			if (1) //RenderTargetTextureRHI != OutputTarget)
+#if MEDIATEXTURERESOURCE_TRACE_RENDER
+			if (Sample.IsValid())
 			{
-				// We use the texture data in the sample
-				CurrentSample = NewCurrentSample;
+				UE_LOG(LogMediaAssets, VeryVerbose, TEXT("TextureResource %p: Sample with time %s cannot be used at time %s"),
+					this,
+					*Sample->GetTime().ToString(TEXT("%h:%m:%s.%t")),
+					*Params.Time.ToString(TEXT("%h:%m:%s.%t"))
+				);
 			}
 			else
 			{
-				// We use the internal copy, no need to hold on to the new one or keep any older one
-				CurrentSample = nullptr;
+				UE_LOG(LogMediaAssets, VeryVerbose, TEXT("TextureResource %p: No valid sample available at time %s"),
+					this,
+					*Params.Time.ToString(TEXT("%h:%m:%s.%t"))
+				);
 			}
-
-			SET_FLOAT_STAT(STAT_MediaUtils_TextureSampleTime, Sample->GetTime().GetTotalMilliseconds());
-		}
-#if MEDIATEXTURERESOURCE_TRACE_RENDER
-		else if (Sample.IsValid())
-		{
-			UE_LOG(LogMediaAssets, VeryVerbose, TEXT("TextureResource %p: Sample with time %s cannot be used at time %s"),
-				this,
-				*Sample->GetTime().ToString(TEXT("%h:%m:%s.%t")),
-				*Params.Time.ToString(TEXT("%h:%m:%s.%t"))
-			);
-		}
-		else
-		{
-			UE_LOG(LogMediaAssets, VeryVerbose, TEXT("TextureResource %p: No valid sample available at time %s"),
-				this,
-				*Params.Time.ToString(TEXT("%h:%m:%s.%t"))
-			);
-		}
 #endif
+		}
 	}
 	else if (Params.CanClear)
 	{
@@ -482,6 +484,9 @@ void FMediaTextureResource::ClearTexture(const FLinearColor& ClearColor, bool Sr
 	// draw the clear color
 	FRHICommandListImmediate& CommandList = FRHICommandListExecutor::GetImmediateCommandList();
 	{
+		SCOPED_DRAW_EVENT(CommandList, FMediaTextureResource_ClearTexture);
+		SCOPED_GPU_STAT(CommandList, MediaTextureResource);
+
 		FRHIRenderPassInfo RPInfo(RenderTargetTextureRHI, ERenderTargetActions::Clear_Store);
 		CommandList.BeginRenderPass(RPInfo, TEXT("ClearTexture"));
 		CommandList.EndRenderPass();
@@ -588,7 +593,8 @@ void FMediaTextureResource::ConvertSample(const TSharedPtr<IMediaTextureSample,
 	// perform the conversion
 	FRHICommandListImmediate& CommandList = FRHICommandListExecutor::GetImmediateCommandList();
 	{
-		SCOPED_DRAW_EVENT(CommandList, MediaTextureConvertResource);
+		SCOPED_DRAW_EVENT(CommandList, FMediaTextureResource_Convert);
+		SCOPED_GPU_STAT(CommandList, MediaTextureResource);
 
 		FGraphicsPipelineStateInitializer GraphicsPSOInit;
 		FRHITexture* RenderTarget = RenderTargetTextureRHI.GetReference();
