MediaFramework: adding capability to generate Mips for MediaTexture output for most platforms to improve rendering quality and perf in 3D use-cases

diff --git a/Engine/Source/Runtime/MediaAssets/Private/Assets/MediaTexture.cpp b/Engine/Source/Runtime/MediaAssets/Private/Assets/MediaTexture.cpp
--- a/Engine/Source/Runtime/MediaAssets/Private/Assets/MediaTexture.cpp
+++ b/Engine/Source/Runtime/MediaAssets/Private/Assets/MediaTexture.cpp
@@ -60,6 +60,7 @@ UMediaTexture::UMediaTexture(const FObjectInitializer& ObjectInitializer)
 	, AddressY(TA_Clamp)
 	, AutoClear(false)
 	, ClearColor(FLinearColor::Black)
+	, NumMips(1)
 	, DefaultGuid(FGuid::NewGuid())
 	, Dimensions(FIntPoint::ZeroValue)
 	, Size(0)
@@ -140,7 +141,9 @@ FTextureResource* UMediaTexture::CreateResource()
 		}
 	}
 
-	return new FMediaTextureResource(*this, Dimensions, Size, ClearColor, CurrentGuid.IsValid() ? CurrentGuid : DefaultGuid);
+	Filter = (NumMips > 1) ? TF_Trilinear : TF_Bilinear;
+
+	return new FMediaTextureResource(*this, Dimensions, Size, ClearColor, CurrentGuid.IsValid() ? CurrentGuid : DefaultGuid, NumMips);
 }
 
 
@@ -337,12 +340,17 @@ void UMediaTexture::TickResource(FTimespan Timecode)
 		return; // retain last frame
 	}
 
+	// update filter state, responding to mips setting
+	Filter = (NumMips > 1) ? TF_Trilinear : TF_Bilinear;
+
+	// setup render parameters
 	RenderParams.CanClear = AutoClear;
 	RenderParams.ClearColor = ClearColor;
 	RenderParams.PreviousGuid = PreviousGuid;
 	RenderParams.CurrentGuid = CurrentGuid;
 	RenderParams.SrgbOutput = SRGB;
-
+	RenderParams.NumMips = NumMips;
+	
 	// redraw texture resource on render thread
 	FMediaTextureResource* ResourceParam = (FMediaTextureResource*)Resource;
 	ENQUEUE_RENDER_COMMAND(MediaTextureResourceRender)(
diff --git a/Engine/Source/Runtime/MediaAssets/Private/Misc/MediaTextureResource.cpp b/Engine/Source/Runtime/MediaAssets/Private/Misc/MediaTextureResource.cpp
--- a/Engine/Source/Runtime/MediaAssets/Private/Misc/MediaTextureResource.cpp
+++ b/Engine/Source/Runtime/MediaAssets/Private/Misc/MediaTextureResource.cpp
@@ -20,6 +20,7 @@
 #include "StaticBoundShaderState.h"
 #include "RenderUtils.h"
 #include "RHIStaticStates.h"
+#include "GenerateMips.h"
 #include "ProfilingDebugging/CsvProfiler.h"
 
 #include "MediaTexture.h"
@@ -27,6 +28,13 @@
 
 #define MEDIATEXTURERESOURCE_TRACE_RENDER 0
 
+// Global mips enabled / disabled switch (more gating is done with GSupportsImageExternal)
+#if PLATFORM_PS4
+# define MTR_USE_MIPS 0
+#else
+# define MTR_USE_MIPS 1
+#endif
+
 
 /** Time spent in media player facade closing media. */
 DECLARE_CYCLE_STAT(TEXT("MediaAssets MediaTextureResource Render"), STAT_MediaAssets_MediaTextureResourceRender, STATGROUP_Media);
@@ -41,7 +49,7 @@ DECLARE_GPU_STAT_NAMED(MediaTextureResource, TEXT("MediaTextureResource"));
 /* Local helpers
  *****************************************************************************/
 
-namespace MediaTextureResource
+namespace MediaTextureResourceHelpers
 {
 	/**
 	 * Get the pixel format for a given sample.
@@ -164,13 +172,19 @@ namespace MediaTextureResource
 /* FMediaTextureResource structors
  *****************************************************************************/
 
-FMediaTextureResource::FMediaTextureResource(UMediaTexture& InOwner, FIntPoint& InOwnerDim, SIZE_T& InOwnerSize, FLinearColor InClearColor, FGuid InTextureGuid)
+FMediaTextureResource::FMediaTextureResource(UMediaTexture& InOwner, FIntPoint& InOwnerDim, SIZE_T& InOwnerSize, FLinearColor InClearColor, FGuid InTextureGuid, uint8 InNumMips)
 	: Cleared(false)
 	, CurrentClearColor(InClearColor)
 	, InitialTextureGuid(InTextureGuid)
 	, Owner(InOwner)
 	, OwnerDim(InOwnerDim)
 	, OwnerSize(InOwnerSize)
+#if MTR_USE_MIPS
+	, CurrentNumMips(GSupportsImageExternal ? 1 : InNumMips)
+#else
+	, CurrentNumMips(1)
+#endif
+	, CurrentSamplerFilter(ESamplerFilter_Num)
 {
 }
 
@@ -237,51 +251,50 @@ void FMediaTextureResource::Render(const FRenderParams& Params)
 #if UE_MEDIAUTILS_DEVELOPMENT_DELEGATE
 		FMediaDelegates::OnPreSampleRender_RenderThread.Broadcast(&Owner, UseSample, Sample);
 #endif
+		// - - - - - - - - - - - - - - - - - - - - - - - - - - - - - - - - - - - -
+		// - - - - - - - - - - - - - - - - - - - - - - - - - - - - - - - - - - - -
+
+#if MTR_USE_MIPS
+		const uint8 NumMips = GSupportsImageExternal ? 1 : Params.NumMips;
+#else
+		const uint8 NumMips = 1;
+#endif
 
 		if (UseSample)
 		{
-			// render the sample
+			//
+			// Valid sample & Sample should be shown
+			//
+
 			if (Sample->GetOutputDim().GetMin() <= 0)
 			{
-#if MEDIATEXTURERESOURCE_TRACE_RENDER
-				UE_LOG(LogMediaAssets, VeryVerbose, TEXT("TextureResource %p: Corrupt sample with time %s at time %s"),
-					this,
-					*Sample->GetTime().ToString(TEXT("%h:%m:%s.%t")),
-					*Params.Time.ToString(TEXT("%h:%m:%s.%t"))
-				);
-#endif
-
+				//
+				// Sample dimensions are invalid
+				//
 				ClearTexture(FLinearColor::Red, Params.SrgbOutput); // mark corrupt sample
 			}
 			else if (Sample->GetMediaTextureSampleConverter())
 			{
-				CreateOutputRenderTarget(Sample, Params);
+				//
+				// Sample brings its own converter
+				//
+				CreateOutputRenderTarget(Sample->GetOutputDim(), MediaTextureResourceHelpers::GetConvertedPixelFormat(Sample), Params.SrgbOutput, Params.ClearColor, NumMips);
 				Sample->GetMediaTextureSampleConverter()->Convert(RenderTargetTextureRHI);
 				Cleared = false;
 			}
-			else if (MediaTextureResource::RequiresConversion(Sample, Params.SrgbOutput))
+			else if (MediaTextureResourceHelpers::RequiresConversion(Sample, Params.SrgbOutput))
 			{
-#if MEDIATEXTURERESOURCE_TRACE_RENDER
-				UE_LOG(LogMediaAssets, VeryVerbose, TEXT("TextureResource %p: Converting sample with time %s at time %s"),
-					this,
-					*Sample->GetTime().ToString(TEXT("%h:%m:%s.%t")),
-					*Params.Time.ToString(TEXT("%h:%m:%s.%t"))
-				);
-#endif
-
-				ConvertSample(Sample, Params.ClearColor, Params.SrgbOutput);
+				//
+				// Sample needs to be converted by built in converter code
+				//
+				ConvertSample(Sample, Params.ClearColor, Params.SrgbOutput, NumMips);
 			}
 			else
 			{
-#if MEDIATEXTURERESOURCE_TRACE_RENDER
-				UE_LOG(LogMediaAssets, VeryVerbose, TEXT("TextureResource %p: Copying sample with time %s at time %s"),
-					this,
-					*Sample->GetTime().ToString(TEXT("%h:%m:%s.%t")),
-					*Params.Time.ToString(TEXT("%h:%m:%s.%t"))
-				);
-#endif
-
-				CopySample(Sample, Params.ClearColor, Params.SrgbOutput);
+				//
+				// Sample can be used directly or is a simple copy
+				//
+				CopySample(Sample, Params.ClearColor, Params.SrgbOutput, NumMips);
 			}
 
 			Rotation = Sample->GetScaleRotation();
@@ -292,10 +305,22 @@ void FMediaTextureResource::Render(const FRenderParams& Params)
 			CurrentSample = TRefCountPtr<FTextureSampleKeeper>(new FTextureSampleKeeper(Sample));
 			check(CurrentSample);
 
+			// Generate mips as needed
+			if (CurrentNumMips > 1 && !Cleared)
+			{
+				check(OutputTarget);
+				FGenerateMips::Execute(FRHICommandListExecutor::GetImmediateCommandList(), OutputTarget, FGenerateMipsParams{ SF_Trilinear, AM_Clamp, AM_Clamp, AM_Clamp });
+			}
+
 			SET_FLOAT_STAT(STAT_MediaUtils_TextureSampleTime, Sample->GetTime().GetTotalMilliseconds());
 		}
 		else
 		{
+			//
+			// Last sample is still valid
+			//
+
+			// Output is using internal buffer?
 			if (OutputTarget == RenderTargetTextureRHI)
 			{
 				/*
@@ -328,28 +353,33 @@ void FMediaTextureResource::Render(const FRenderParams& Params)
 	}
 	else if (Params.CanClear)
 	{
+		//
+		// No valid sample source & we should clear
+		//
+
+		// Need to clear the output?
 		if (!Cleared || (Params.ClearColor != CurrentClearColor))
 		{
-#if MEDIATEXTURERESOURCE_TRACE_RENDER
-			UE_LOG(LogMediaAssets, VeryVerbose, TEXT("TextureResource %p: Clearing texture at time %s"),
-				this,
-				*Params.Time.ToString(TEXT("%h:%m:%s.%t"))
-			);
-#endif
-
+			// Yes...
 			ClearTexture(Params.ClearColor, Params.SrgbOutput);
 
 			// Also get rid of any sample from previous rendering...
 			CurrentSample = nullptr;
 		}
 	}
-	
-	//Cache next available sample time in the MediaTexture owner since we're the only one that can consume from the queue
+
+	// - - - - - - - - - - - - - - - - - - - - - - - - - - - - - - - - - - - -
+	// - - - - - - - - - - - - - - - - - - - - - - - - - - - - - - - - - - - -
+
+	// Cache next available sample time in the MediaTexture owner since we're the only one that can consume from the queue
 	CacheNextAvailableSampleTime(SampleSource);
 
-	// update external texture registration
+	// Update external texture registration in case we have no native support
+	// (in that case there is support, the player will do this - but it is used all the time)
 	if (!GSupportsImageExternal)
 	{
+		SetupSampler();
+
 		if (Params.CurrentGuid.IsValid())
 		{
 			FTextureRHIRef VideoTexture = (FTextureRHIRef)Owner.TextureReference.TextureReferenceRHI;
@@ -362,7 +392,7 @@ void FMediaTextureResource::Render(const FRenderParams& Params)
 		}
 	}
 	
-	//Update usable Guid for the RenderThread
+	// Update usable Guid for the RenderThread
 	Owner.SetRenderedExternalTextureGuid(Params.CurrentGuid);
 }
 
@@ -397,18 +427,36 @@ uint32 FMediaTextureResource::GetSizeY() const
 }
 
 
-void FMediaTextureResource::InitDynamicRHI()
+void FMediaTextureResource::SetupSampler()
 {
-	// create the sampler state
-	FSamplerStateInitializerRHI SamplerStateInitializer(
-		(ESamplerFilter)UDeviceProfileManager::Get().GetActiveProfile()->GetTextureLODSettings()->GetSamplerFilter(&Owner),
-		(Owner.AddressX == TA_Wrap) ? AM_Wrap : ((Owner.AddressX == TA_Clamp) ? AM_Clamp : AM_Mirror),
-		(Owner.AddressY == TA_Wrap) ? AM_Wrap : ((Owner.AddressY == TA_Clamp) ? AM_Clamp : AM_Mirror),
-		AM_Wrap
-	);
+#if MTR_USE_MIPS
+	// note: no need to check if we support external textures & suppress mips as this sampler state will not be used in that case
+	ESamplerFilter OwnerFilter = (ESamplerFilter)UDeviceProfileManager::Get().GetActiveProfile()->GetTextureLODSettings()->GetSamplerFilter(&Owner);
+#else
+	ESamplerFilter OwnerFilter = SF_Bilinear;
+#endif
+
+	if (CurrentSamplerFilter != OwnerFilter)
+	{
+		CurrentSamplerFilter = OwnerFilter;
+
+		// create the sampler state
+		FSamplerStateInitializerRHI SamplerStateInitializer(
+			CurrentSamplerFilter,
+			(Owner.AddressX == TA_Wrap) ? AM_Wrap : ((Owner.AddressX == TA_Clamp) ? AM_Clamp : AM_Mirror),
+			(Owner.AddressY == TA_Wrap) ? AM_Wrap : ((Owner.AddressY == TA_Clamp) ? AM_Clamp : AM_Mirror),
+			AM_Wrap
+		);
+
+		SamplerStateRHI = RHICreateSamplerState(SamplerStateInitializer);
+	}
+}
 
-	SamplerStateRHI = RHICreateSamplerState(SamplerStateInitializer);
 
+void FMediaTextureResource::InitDynamicRHI()
+{
+	SetupSampler();
+	
 	// Note: set up default texture, or we can get sampler bind errors on render
 	// we can't leave here without having a valid bindable resource for some RHIs.
 
@@ -419,6 +467,7 @@ void FMediaTextureResource::InitDynamicRHI()
 	check(RenderTargetTextureRHI.IsValid());
 	check(OutputTarget.IsValid());
 
+	// Register "external texture" parameters if the platform does not support them (and hence the player does not set them)
 	if (!GSupportsImageExternal)
 	{
 		FTextureRHIRef VideoTexture = (FTextureRHIRef)Owner.TextureReference.TextureReferenceRHI;
@@ -446,40 +495,7 @@ void FMediaTextureResource::ReleaseDynamicRHI()
 void FMediaTextureResource::ClearTexture(const FLinearColor& ClearColor, bool SrgbOutput)
 {
 	// create output render target if we don't have one yet
-	const uint32 OutputCreateFlags = TexCreate_Dynamic | (SrgbOutput ? TexCreate_SRGB : 0);
-	const EPixelFormat OutputPixelFormat = PF_B8G8R8A8;
-
-	if ((ClearColor != CurrentClearColor) || !OutputTarget.IsValid() || (OutputTarget->GetFormat() != OutputPixelFormat) || ((OutputTarget->GetFlags() & OutputCreateFlags) != OutputCreateFlags))
-	{
-		FString DebugName = Owner.GetName();
-
-		FRHIResourceCreateInfo CreateInfo;
-		CreateInfo.ClearValueBinding = FClearValueBinding(ClearColor);
-		CreateInfo.DebugName = *DebugName;
-
-		TRefCountPtr<FRHITexture2D> DummyTexture2DRHI;
-
-		RHICreateTargetableShaderResource2D(
-			2,
-			2,
-			OutputPixelFormat,
-			1,
-			OutputCreateFlags,
-			TexCreate_RenderTargetable,
-			false,
-			CreateInfo,
-			OutputTarget,
-			DummyTexture2DRHI
-		);
-
-		CurrentClearColor = ClearColor;
-		UpdateResourceSize();
-	}
-
-	if (RenderTargetTextureRHI != OutputTarget)
-	{
-		UpdateTextureReference(OutputTarget);
-	}
+	CreateOutputRenderTarget(FIntPoint(2, 2), PF_B8G8R8A8, SrgbOutput, ClearColor, 1);
 
 	// draw the clear color
 	FRHICommandListImmediate& CommandList = FRHICommandListExecutor::GetImmediateCommandList();
@@ -498,9 +514,9 @@ void FMediaTextureResource::ClearTexture(const FLinearColor& ClearColor, bool Sr
 }
 
 
-void FMediaTextureResource::ConvertSample(const TSharedPtr<IMediaTextureSample, ESPMode::ThreadSafe>& Sample, const FLinearColor& ClearColor, bool SrgbOutput)
+void FMediaTextureResource::ConvertSample(const TSharedPtr<IMediaTextureSample, ESPMode::ThreadSafe>& Sample, const FLinearColor& ClearColor, bool SrgbOutput, uint8 InNumMips)
 {
-	const EPixelFormat InputPixelFormat = MediaTextureResource::GetPixelFormat(Sample);
+	const EPixelFormat InputPixelFormat = MediaTextureResourceHelpers::GetPixelFormat(Sample);
 
 	// get input texture
 	FRHITexture2D* InputTexture = nullptr;
@@ -514,6 +530,8 @@ void FMediaTextureResource::ConvertSample(const TSharedPtr<IMediaTextureSample,
 
 		if (SampleTexture2D != nullptr)
 		{
+			// Use the sample as source texture...
+
 			InputTexture = SampleTexture2D;
 
 			InputTarget.SafeRelease();
@@ -521,11 +539,13 @@ void FMediaTextureResource::ConvertSample(const TSharedPtr<IMediaTextureSample,
 		}
 		else
 		{
-			const bool SrgbTexture = MediaTextureResource::RequiresSrgbTexture(Sample);
+			// Make a source texture so we can convert from it...
+
+			const bool SrgbTexture = MediaTextureResourceHelpers::RequiresSrgbTexture(Sample);
 			const uint32 InputCreateFlags = TexCreate_Dynamic | (SrgbTexture ? TexCreate_SRGB : 0);
 			const FIntPoint SampleDim = Sample->GetDim();
 
-			// create a new input render target if necessary
+			// create a new temp input render target if necessary
 			if (!InputTarget.IsValid() || (InputTarget->GetSizeXY() != SampleDim) || (InputTarget->GetFormat() != InputPixelFormat) || ((InputTarget->GetFlags() & InputCreateFlags) != InputCreateFlags))
 			{
 				TRefCountPtr<FRHITexture2D> DummyTexture2DRHI;
@@ -555,40 +575,9 @@ void FMediaTextureResource::ConvertSample(const TSharedPtr<IMediaTextureSample,
 		}
 	}
 
-	// create output render target if necessary
-	const uint32 OutputCreateFlags = TexCreate_Dynamic | (SrgbOutput ? TexCreate_SRGB : 0);
+	// create the output texture
 	const FIntPoint OutputDim = Sample->GetOutputDim();
-	const EPixelFormat OutputPixelFormat = MediaTextureResource::GetConvertedPixelFormat(Sample);
-
-	if ((ClearColor != CurrentClearColor) || !OutputTarget.IsValid() || (OutputTarget->GetSizeXY() != OutputDim) || (OutputTarget->GetFormat() != OutputPixelFormat) || ((OutputTarget->GetFlags() & OutputCreateFlags) != OutputCreateFlags))
-	{
-		TRefCountPtr<FRHITexture2D> DummyTexture2DRHI;
-		
-		FRHIResourceCreateInfo CreateInfo = {
-			FClearValueBinding(ClearColor)
-		};
-
-		RHICreateTargetableShaderResource2D(
-			OutputDim.X,
-			OutputDim.Y,
-			OutputPixelFormat,
-			1,
-			OutputCreateFlags,
-			TexCreate_RenderTargetable,
-			false,
-			CreateInfo,
-			OutputTarget,
-			DummyTexture2DRHI
-		);
-
-		CurrentClearColor = ClearColor;
-		UpdateResourceSize();
-	}
-
-	if (RenderTargetTextureRHI != OutputTarget)
-	{
-		UpdateTextureReference(OutputTarget);
-	}
+	CreateOutputRenderTarget(OutputDim, MediaTextureResourceHelpers::GetConvertedPixelFormat(Sample), SrgbOutput, ClearColor, InNumMips);
 
 	// perform the conversion
 	FRHICommandListImmediate& CommandList = FRHICommandListExecutor::GetImmediateCommandList();
@@ -599,7 +588,7 @@ void FMediaTextureResource::ConvertSample(const TSharedPtr<IMediaTextureSample,
 		FGraphicsPipelineStateInitializer GraphicsPSOInit;
 		FRHITexture* RenderTarget = RenderTargetTextureRHI.GetReference();
 
-		FRHIRenderPassInfo RPInfo(RenderTarget, ERenderTargetActions::Load_Store);
+		FRHIRenderPassInfo RPInfo(RenderTarget, ERenderTargetActions::DontLoad_Store);
 		CommandList.BeginRenderPass(RPInfo, TEXT("ConvertMedia"));
 		{
 			CommandList.ApplyCachedRenderTargets(GraphicsPSOInit);
@@ -748,7 +737,7 @@ void FMediaTextureResource::ConvertSample(const TSharedPtr<IMediaTextureSample,
 }
 
 
-void FMediaTextureResource::CopySample(const TSharedPtr<IMediaTextureSample, ESPMode::ThreadSafe>& Sample, const FLinearColor& ClearColor, bool SrgbOutput)
+void FMediaTextureResource::CopySample(const TSharedPtr<IMediaTextureSample, ESPMode::ThreadSafe>& Sample, const FLinearColor& ClearColor, bool SrgbOutput, uint8 InNumMips)
 {
 	FRHITexture* SampleTexture = Sample->GetTexture();
 	FRHITexture2D* SampleTexture2D = (SampleTexture != nullptr) ? SampleTexture->GetTexture2D() : nullptr;
@@ -759,7 +748,9 @@ void FMediaTextureResource::CopySample(const TSharedPtr<IMediaTextureSample, ESP
 
 	if (SampleTexture2D != nullptr)
 	{
-		// use sample's texture as the new render target.
+		if (InNumMips == 1)
+		{
+			// Use sample's texture as the new render target.
 		if (TextureRHI != SampleTexture2D)
 		{
 			UpdateTextureReference(SampleTexture2D);
@@ -767,50 +758,28 @@ void FMediaTextureResource::CopySample(const TSharedPtr<IMediaTextureSample, ESP
 			OutputTarget.SafeRelease();
 			UpdateResourceSize();
 		}
+
 	}
 	else
 	{
-		// create a new output render target if necessary
-		const uint32 OutputCreateFlags = TexCreate_Dynamic | (SrgbOutput ? TexCreate_SRGB : 0);
-		const EPixelFormat SampleFormat = MediaTextureResource::GetPixelFormat(Sample);
-		const FIntPoint SampleDim = Sample->GetDim();
+			// Texture to receive texture from sample
+			CreateOutputRenderTarget(Sample->GetOutputDim(), MediaTextureResourceHelpers::GetPixelFormat(Sample), SrgbOutput, ClearColor, InNumMips);
 
-		if ((ClearColor != CurrentClearColor) || !OutputTarget.IsValid() || (OutputTarget->GetSizeXY() != SampleDim) || (OutputTarget->GetFormat() != SampleFormat) || ((OutputTarget->GetFlags() & OutputCreateFlags) != OutputCreateFlags))
-		{
-			TRefCountPtr<FRHITexture2D> DummyTexture2DRHI;
-
-			FRHIResourceCreateInfo CreateInfo = {
-				FClearValueBinding(ClearColor)
-			};
-
-			RHICreateTargetableShaderResource2D(
-				SampleDim.X,
-				SampleDim.Y,
-				SampleFormat,
-				1,
-				OutputCreateFlags,
-				TexCreate_RenderTargetable,
-				false,
-				CreateInfo,
-				OutputTarget,
-				DummyTexture2DRHI
-			);
-
-			CurrentClearColor = ClearColor;
-			UpdateResourceSize();
+			// Copy data into the output texture to able to add mips later on
+			FRHICommandListExecutor::GetImmediateCommandList().CopyTexture(SampleTexture2D, OutputTarget, FRHICopyTextureInfo());
 		}
 
-		if (RenderTargetTextureRHI != OutputTarget)
-		{
-			UpdateTextureReference(OutputTarget);
+		Cleared = false;
 		}
+	else
+		{
+		// Texture to receive precisely only output pixels via CPU copy
+		CreateOutputRenderTarget(Sample->GetDim(), MediaTextureResourceHelpers::GetPixelFormat(Sample), SrgbOutput, ClearColor, InNumMips);
 
 		// copy sample data to output render target
-		FUpdateTextureRegion2D Region(0, 0, 0, 0, SampleDim.X, SampleDim.Y);
+		FUpdateTextureRegion2D Region(0, 0, 0, 0, Sample->GetDim().X, Sample->GetDim().Y);
 		RHIUpdateTexture2D(RenderTargetTextureRHI.GetReference(), 0, Region, Sample->GetStride(), (uint8*)Sample->GetBuffer());
 	}
-
-	Cleared = false;
 }
 
 
@@ -849,26 +818,34 @@ void FMediaTextureResource::UpdateTextureReference(FRHITexture2D* NewTexture)
 	}
 }
 
-void FMediaTextureResource::CreateOutputRenderTarget(const TSharedPtr<IMediaTextureSample, ESPMode::ThreadSafe>& InSample, const FRenderParams& InParams)
+
+void FMediaTextureResource::CreateOutputRenderTarget(const FIntPoint & InDim, EPixelFormat InPixelFormat, bool bInSRGB, const FLinearColor & InClearColor, uint8 InNumMips)
 {
 	// create output render target if necessary
-	const uint32 OutputCreateFlags = TexCreate_Dynamic | (InParams.SrgbOutput ? TexCreate_SRGB : 0);
-	const FIntPoint OutputDim = InSample->GetOutputDim();
-	const EPixelFormat OutputPixelFormat = MediaTextureResource::GetConvertedPixelFormat(InSample);
+	uint32 OutputCreateFlags = TexCreate_Dynamic | (bInSRGB ? TexCreate_SRGB : 0);
+	if (InNumMips > 1)
+	{
+		// Make sure can have mips & the mip generator has what it needs to work
+		OutputCreateFlags |= (TexCreate_GenerateMipCapable | TexCreate_UAV);
 
-	if ((InParams.ClearColor != CurrentClearColor) || !OutputTarget.IsValid() || (OutputTarget->GetSizeXY() != OutputDim) || (OutputTarget->GetFormat() != OutputPixelFormat) || ((OutputTarget->GetFlags() & OutputCreateFlags) != OutputCreateFlags))
+		// Make sure we only set a number of mips that actually makes sense, given the sample size
+		uint8 MaxMips = FGenericPlatformMath::FloorToInt(FGenericPlatformMath::Log2(FGenericPlatformMath::Max(InDim.X, InDim.Y)));
+		InNumMips = FMath::Min(InNumMips, MaxMips);
+	}
+
+	if ((InClearColor != CurrentClearColor) || !OutputTarget.IsValid() || (OutputTarget->GetSizeXY() != InDim) || (OutputTarget->GetFormat() != InPixelFormat) || ((OutputTarget->GetFlags() & OutputCreateFlags) != OutputCreateFlags) || CurrentNumMips != InNumMips)
 	{
 		TRefCountPtr<FRHITexture2D> DummyTexture2DRHI;
 
 		FRHIResourceCreateInfo CreateInfo = {
-			FClearValueBinding(InParams.ClearColor)
+			FClearValueBinding(InClearColor)
 		};
 
 		RHICreateTargetableShaderResource2D(
-			OutputDim.X,
-			OutputDim.Y,
-			OutputPixelFormat,
-			1,
+			InDim.X,
+			InDim.Y,
+			InPixelFormat,
+			InNumMips,
 			OutputCreateFlags,
 			TexCreate_RenderTargetable,
 			false,
@@ -877,7 +854,8 @@ void FMediaTextureResource::CreateOutputRenderTarget(const TSharedPtr<IMediaText
 			DummyTexture2DRHI
 		);
 
-		CurrentClearColor = InParams.ClearColor;
+		CurrentClearColor = InClearColor;
+		CurrentNumMips = InNumMips;
 		UpdateResourceSize();
 
 		Cleared = false;
@@ -889,6 +867,7 @@ void FMediaTextureResource::CreateOutputRenderTarget(const TSharedPtr<IMediaText
 	}
 }
 
+
 void FMediaTextureResource::CacheNextAvailableSampleTime(const TSharedPtr<FMediaTextureSampleSource, ESPMode::ThreadSafe>& InSampleQueue) const
 {
 	FTimespan SampleTime(FTimespan::MinValue());
diff --git a/Engine/Source/Runtime/MediaAssets/Private/Misc/MediaTextureResource.h b/Engine/Source/Runtime/MediaAssets/Private/Misc/MediaTextureResource.h
--- a/Engine/Source/Runtime/MediaAssets/Private/Misc/MediaTextureResource.h
+++ b/Engine/Source/Runtime/MediaAssets/Private/Misc/MediaTextureResource.h
@@ -39,8 +39,9 @@ public:
 	 * @param InOWnerSize Reference to the size in bytes of the texture that owns this resource (will be updated by resource).
 	 * @param InClearColor The initial clear color.
 	 * @param InTextureGuid The initial external texture GUID.
+	 * @param InNumMips The initial number of mips to be generated for the output texture
 	 */
-	FMediaTextureResource(UMediaTexture& InOwner, FIntPoint& InOwnerDim, SIZE_T& InOwnerSize, FLinearColor InClearColor, FGuid InTextureGuid);
+	FMediaTextureResource(UMediaTexture& InOwner, FIntPoint& InOwnerDim, SIZE_T& InOwnerSize, FLinearColor InClearColor, FGuid InTextureGuid, uint8 InNumMips);
 
 	/** Virtual destructor. */
 	virtual ~FMediaTextureResource() { }
@@ -71,6 +72,9 @@ public:
 		/** Whether output should be in sRGB color space. */
 		bool SrgbOutput;
 
+		/** Number of mips wanted */
+		uint8 NumMips;
+
 		/** The time of the video frame to render (in player's clock). */
 		FTimespan Time;
 	};
@@ -117,9 +121,10 @@ protected:
 	 * @param Sample The texture sample to convert.
 	 * @param ClearColor The clear color to use for the output texture.
 	 * @param SrgbOutput Whether the output texture is in sRGB color space.
+	 * @param Number of mips
 	 * @see CopySample
 	 */
-	void ConvertSample(const TSharedPtr<IMediaTextureSample, ESPMode::ThreadSafe>& Sample, const FLinearColor& ClearColor, bool SrgbOutput);
+	void ConvertSample(const TSharedPtr<IMediaTextureSample, ESPMode::ThreadSafe>& Sample, const FLinearColor& ClearColor, bool SrgbOutput, uint8 InNumMips);
 
 	/**
 	 * Render the given texture sample by using it as or copying it to the render target.
@@ -127,9 +132,10 @@ protected:
 	 * @param Sample The texture sample to copy.
 	 * @param ClearColor The clear color to use for the output texture.
 	 * @param SrgbOutput Whether the output texture is in sRGB color space.
+	 * @param Number of mips
 	 * @see ConvertSample
 	 */
-	void CopySample(const TSharedPtr<IMediaTextureSample, ESPMode::ThreadSafe>& Sample, const FLinearColor& ClearColor, bool SrgbOutput);
+	void CopySample(const TSharedPtr<IMediaTextureSample, ESPMode::ThreadSafe>& Sample, const FLinearColor& ClearColor, bool SrgbOutput, uint8 InNumMips);
 
 	/** Calculates the current resource size and notifies the owner texture. */
 	void UpdateResourceSize();
@@ -143,11 +149,8 @@ protected:
 
 	/**
 	 * Create/update output render target as needed
-	 *
-	 * @param InSample Sample to query render target dimension and format from
-	 * @param InParams Parameters containing SrgbOutput and ClearColor
 	 */
-	void CreateOutputRenderTarget(const TSharedPtr<IMediaTextureSample, ESPMode::ThreadSafe>& InSample, const FRenderParams& InParams);
+	void CreateOutputRenderTarget(const FIntPoint & InDim, EPixelFormat InPixelFormat, bool bInSRGB, const FLinearColor & InClearColor, uint8 InNumMips);
 
 	/**
 	 * Caches next available sample from queue in MediaTexture owner to keep single consumer access
@@ -156,6 +159,9 @@ protected:
 	 */
 	void CacheNextAvailableSampleTime(const TSharedPtr<FMediaTextureSampleSource, ESPMode::ThreadSafe>& InSampleQueue) const;
 
+	/** Setup sampler state from owner's settings as needed */
+	void SetupSampler();
+
 private:
 
 	// Class to maintain a reference to a IMediaTextureSample instance as long  as needed by RHI
@@ -195,6 +201,12 @@ private:
 	/** Reference to the owner's texture size field. */
 	SIZE_T& OwnerSize;
 
+	/** Current number of mips to be generated as output */
+	uint8 CurrentNumMips;
+
+	/** Current texture sampler filter value */
+	ESamplerFilter CurrentSamplerFilter;
+
 	/** The current media player facade to get video samples from. */
 	TWeakPtr<FMediaPlayerFacade, ESPMode::ThreadSafe> PlayerFacadePtr;
 
diff --git a/Engine/Source/Runtime/MediaAssets/Public/MediaTexture.h b/Engine/Source/Runtime/MediaAssets/Public/MediaTexture.h
--- a/Engine/Source/Runtime/MediaAssets/Public/MediaTexture.h
+++ b/Engine/Source/Runtime/MediaAssets/Public/MediaTexture.h
@@ -47,6 +47,10 @@ class MEDIAASSETS_API UMediaTexture
 	UPROPERTY(EditAnywhere, BlueprintReadWrite, Category="MediaTexture")
 	FLinearColor ClearColor;
 
+	/** The number of mips to use (default = 1). */
+	UPROPERTY(EditAnywhere, BlueprintReadWrite, Category = "MediaTexture", meta=(DisplayName="Total number of Mipmaps to output"))
+	uint8 NumMips;
+
 public:
 
 	/**
