MediaFramework: fixing up FMemStack usage in case mips are generated from video output texture
#rnx

diff --git a/Engine/Source/Runtime/MediaAssets/Private/Misc/MediaTextureResource.cpp b/Engine/Source/Runtime/MediaAssets/Private/Misc/MediaTextureResource.cpp
--- a/Engine/Source/Runtime/MediaAssets/Private/Misc/MediaTextureResource.cpp
+++ b/Engine/Source/Runtime/MediaAssets/Private/Misc/MediaTextureResource.cpp
@@ -309,6 +309,8 @@ void FMediaTextureResource::Render(const FRenderParams& Params)
 			if (CurrentNumMips > 1 && !Cleared)
 			{
 				check(OutputTarget);
+
+				FMemMark MemMark(FMemStack::Get());
 				FGenerateMips::Execute(FRHICommandListExecutor::GetImmediateCommandList(), OutputTarget, FGenerateMipsParams{ SF_Trilinear, AM_Clamp, AM_Clamp, AM_Clamp });
 			}
 
