Enabling optional mips generation for videos on PS4
#rnx

diff --git a/Engine/Source/Runtime/MediaAssets/Private/Misc/MediaTextureResource.cpp b/Engine/Source/Runtime/MediaAssets/Private/Misc/MediaTextureResource.cpp
--- a/Engine/Source/Runtime/MediaAssets/Private/Misc/MediaTextureResource.cpp
+++ b/Engine/Source/Runtime/MediaAssets/Private/Misc/MediaTextureResource.cpp
@@ -29,11 +29,7 @@
 #define MEDIATEXTURERESOURCE_TRACE_RENDER 0
 
 // Global mips enabled / disabled switch (more gating is done with GSupportsImageExternal)
-#if PLATFORM_PS4
-# define MTR_USE_MIPS 0
-#else
 # define MTR_USE_MIPS 1
-#endif
 
 
 /** Time spent in media player facade closing media. */
