MediaFramework: clenaup around mip generation resource state transitions
#rnx

diff --git a/Engine/Source/Runtime/MediaAssets/Private/Misc/MediaTextureResource.cpp b/Engine/Source/Runtime/MediaAssets/Private/Misc/MediaTextureResource.cpp
--- a/Engine/Source/Runtime/MediaAssets/Private/Misc/MediaTextureResource.cpp
+++ b/Engine/Source/Runtime/MediaAssets/Private/Misc/MediaTextureResource.cpp
@@ -29,7 +29,7 @@
 #define MEDIATEXTURERESOURCE_TRACE_RENDER 0
 
 // Global mips enabled / disabled switch (more gating is done with GSupportsImageExternal)
-# define MTR_USE_MIPS 1
+#define MTR_USE_MIPS 1
 
 
 /** Time spent in media player facade closing media. */
@@ -307,7 +307,11 @@ void FMediaTextureResource::Render(const FRenderParams& Params)
 				check(OutputTarget);
 
 				FMemMark MemMark(FMemStack::Get());
-				FGenerateMips::Execute(FRHICommandListExecutor::GetImmediateCommandList(), OutputTarget, FGenerateMipsParams{ SF_Trilinear, AM_Clamp, AM_Clamp, AM_Clamp });
+				FRHICommandListImmediate & RHICmdList = FRHICommandListExecutor::GetImmediateCommandList();
+
+				RHICmdList.TransitionResource(EResourceTransitionAccess::ERWBarrier, OutputTarget);
+				FGenerateMips::Execute(RHICmdList, OutputTarget, FGenerateMipsParams{ SF_Bilinear, AM_Clamp, AM_Clamp, AM_Clamp });
+				RHICmdList.TransitionResource(EResourceTransitionAccess::EReadable, OutputTarget);
 			}
 
 			SET_FLOAT_STAT(STAT_MediaUtils_TextureSampleTime, Sample->GetTime().GetTotalMilliseconds());
@@ -329,24 +333,6 @@ void FMediaTextureResource::Render(const FRenderParams& Params)
 				*/
 				CurrentSample = nullptr;
 			}
-
-#if MEDIATEXTURERESOURCE_TRACE_RENDER
-			if (Sample.IsValid())
-			{
-				UE_LOG(LogMediaAssets, VeryVerbose, TEXT("TextureResource %p: Sample with time %s cannot be used at time %s"),
-					this,
-					*Sample->GetTime().ToString(TEXT("%h:%m:%s.%t")),
-					*Params.Time.ToString(TEXT("%h:%m:%s.%t"))
-				);
-			}
-			else
-			{
-				UE_LOG(LogMediaAssets, VeryVerbose, TEXT("TextureResource %p: No valid sample available at time %s"),
-					this,
-					*Params.Time.ToString(TEXT("%h:%m:%s.%t"))
-				);
-			}
-#endif
 		}
 	}
 	else if (Params.CanClear)
@@ -360,7 +346,6 @@ void FMediaTextureResource::Render(const FRenderParams& Params)
 		{
 			// Yes...
 			ClearTexture(Params.ClearColor, Params.SrgbOutput);
-
 			// Also get rid of any sample from previous rendering...
 			CurrentSample = nullptr;
 		}
