Enabling mips generation for Android & adding explicit per MediaTexture instance on/off switch to mips generation
#rnx

diff --git a/Engine/Shaders/Private/MediaShaders.usf b/Engine/Shaders/Private/MediaShaders.usf
--- a/Engine/Shaders/Private/MediaShaders.usf
+++ b/Engine/Shaders/Private/MediaShaders.usf
@@ -369,3 +369,14 @@ void SetAlphaOnePS(
 	OutColor.xyzw = float4(Color.rgb, 1.0f);
 }
 
+
+// shader to read from TextureExternal source
+void ReadTextureExternalPS(
+	float4 InPosition : SV_POSITION,
+	float2 InUV : TEXCOORD0,
+	out float4 OutColor : SV_Target0)
+{
+    float2 UV = float2(dot(InUV, ReadTextureExternalUB.ScaleRotation.xy), dot(InUV, ReadTextureExternalUB.ScaleRotation.zw)) + ReadTextureExternalUB.Offset;
+	float4 InColor = ReadTextureExternalUB.Texture.Sample(ReadTextureExternalUB.SamplerP, UV);
+    OutColor = float4(sRGBToLinear(InColor.xyz), InColor.w);
+}
diff --git a/Engine/Source/Runtime/MediaAssets/Private/Assets/MediaTexture.cpp b/Engine/Source/Runtime/MediaAssets/Private/Assets/MediaTexture.cpp
--- a/Engine/Source/Runtime/MediaAssets/Private/Assets/MediaTexture.cpp
+++ b/Engine/Source/Runtime/MediaAssets/Private/Assets/MediaTexture.cpp
@@ -60,6 +60,7 @@ UMediaTexture::UMediaTexture(const FObjectInitializer& ObjectInitializer)
 	, AddressY(TA_Clamp)
 	, AutoClear(false)
 	, ClearColor(FLinearColor::Black)
+	, EnableGenMips(false)
 	, NumMips(1)
 	, DefaultGuid(FGuid::NewGuid())
 	, Dimensions(FIntPoint::ZeroValue)
@@ -141,15 +142,15 @@ FTextureResource* UMediaTexture::CreateResource()
 		}
 	}
 
-	Filter = (NumMips > 1) ? TF_Trilinear : TF_Bilinear;
+	Filter = (EnableGenMips && (NumMips > 1)) ? TF_Trilinear : TF_Bilinear;
 
-	return new FMediaTextureResource(*this, Dimensions, Size, ClearColor, CurrentGuid.IsValid() ? CurrentGuid : DefaultGuid, NumMips);
+	return new FMediaTextureResource(*this, Dimensions, Size, ClearColor, CurrentGuid.IsValid() ? CurrentGuid : DefaultGuid, EnableGenMips, NumMips);
 }
 
 
 EMaterialValueType UMediaTexture::GetMaterialType() const
 {
-	return MCT_TextureExternal;
+	return EnableGenMips ? MCT_Texture2D : MCT_TextureExternal;
 }
 
 
@@ -167,6 +168,10 @@ float UMediaTexture::GetSurfaceHeight() const
 
 FGuid UMediaTexture::GetExternalTextureGuid() const
 {
+	if (EnableGenMips)
+	{
+		return FGuid();
+	}
 	FScopeLock Lock(&CriticalSection);
 	return CurrentRenderedGuid;
 }
@@ -341,7 +346,7 @@ void UMediaTexture::TickResource(FTimespan Timecode)
 	}
 
 	// update filter state, responding to mips setting
-	Filter = (NumMips > 1) ? TF_Trilinear : TF_Bilinear;
+	Filter = (EnableGenMips && (NumMips > 1)) ? TF_Trilinear : TF_Bilinear;
 
 	// setup render parameters
 	RenderParams.CanClear = AutoClear;
diff --git a/Engine/Source/Runtime/MediaAssets/Private/Misc/MediaTextureResource.cpp b/Engine/Source/Runtime/MediaAssets/Private/Misc/MediaTextureResource.cpp
--- a/Engine/Source/Runtime/MediaAssets/Private/Misc/MediaTextureResource.cpp
+++ b/Engine/Source/Runtime/MediaAssets/Private/Misc/MediaTextureResource.cpp
@@ -25,13 +25,12 @@
 
 #include "MediaTexture.h"
 
+#if PLATFORM_ANDROID
+# include "Android/AndroidPlatformMisc.h"
+#endif
 
 #define MEDIATEXTURERESOURCE_TRACE_RENDER 0
 
-// Global mips enabled / disabled switch (more gating is done with GSupportsImageExternal)
-#define MTR_USE_MIPS 1
-
-
 /** Time spent in media player facade closing media. */
 DECLARE_CYCLE_STAT(TEXT("MediaAssets MediaTextureResource Render"), STAT_MediaAssets_MediaTextureResourceRender, STATGROUP_Media);
 
@@ -168,20 +167,22 @@ namespace MediaTextureResourceHelpers
 /* FMediaTextureResource structors
  *****************************************************************************/
 
-FMediaTextureResource::FMediaTextureResource(UMediaTexture& InOwner, FIntPoint& InOwnerDim, SIZE_T& InOwnerSize, FLinearColor InClearColor, FGuid InTextureGuid, uint8 InNumMips)
+FMediaTextureResource::FMediaTextureResource(UMediaTexture& InOwner, FIntPoint& InOwnerDim, SIZE_T& InOwnerSize, FLinearColor InClearColor, FGuid InTextureGuid, bool InEnableGenMips, uint8 InNumMips)
 	: Cleared(false)
 	, CurrentClearColor(InClearColor)
 	, InitialTextureGuid(InTextureGuid)
 	, Owner(InOwner)
 	, OwnerDim(InOwnerDim)
 	, OwnerSize(InOwnerSize)
-#if MTR_USE_MIPS
-	, CurrentNumMips(GSupportsImageExternal ? 1 : InNumMips)
-#else
-	, CurrentNumMips(1)
-#endif
+	, bEnableGenMips(InEnableGenMips)
+	, CurrentNumMips(InEnableGenMips ? InNumMips : 1)
 	, CurrentSamplerFilter(ESamplerFilter_Num)
 {
+#if PLATFORM_ANDROID
+	bUsesImageExternal = (!FAndroidMisc::ShouldUseVulkan() && GSupportsImageExternal);
+#else
+	bUsesImageExternal = GSupportsImageExternal;
+#endif
 }
 
 
@@ -250,13 +251,11 @@ void FMediaTextureResource::Render(const FRenderParams& Params)
 		// - - - - - - - - - - - - - - - - - - - - - - - - - - - - - - - - - - - -
 		// - - - - - - - - - - - - - - - - - - - - - - - - - - - - - - - - - - - -
 
-#if MTR_USE_MIPS
-		const uint8 NumMips = GSupportsImageExternal ? 1 : Params.NumMips;
-#else
-		const uint8 NumMips = 1;
-#endif
+		const uint8 NumMips = bEnableGenMips ? Params.NumMips : 1;
 
-		if (UseSample)
+		// If real "external texture" support is in place and no mips are used the image will "bypass" any of this processing via the GUID-based lookup for "ExternalTextures" and will
+		// reach the reading material shader without any processing here...
+		if (UseSample && !(bUsesImageExternal && !bEnableGenMips))
 		{
 			//
 			// Valid sample & Sample should be shown
@@ -290,7 +289,7 @@ void FMediaTextureResource::Render(const FRenderParams& Params)
 				//
 				// Sample can be used directly or is a simple copy
 				//
-				CopySample(Sample, Params.ClearColor, Params.SrgbOutput, NumMips);
+				CopySample(Sample, Params.ClearColor, Params.SrgbOutput, NumMips, Params.CurrentGuid);
 			}
 
 			Rotation = Sample->GetScaleRotation();
@@ -359,7 +358,7 @@ void FMediaTextureResource::Render(const FRenderParams& Params)
 
 	// Update external texture registration in case we have no native support
 	// (in that case there is support, the player will do this - but it is used all the time)
-	if (!GSupportsImageExternal)
+	if (!bUsesImageExternal)
 	{
 		SetupSampler();
 
@@ -412,12 +411,7 @@ uint32 FMediaTextureResource::GetSizeY() const
 
 void FMediaTextureResource::SetupSampler()
 {
-#if MTR_USE_MIPS
-	// note: no need to check if we support external textures & suppress mips as this sampler state will not be used in that case
-	ESamplerFilter OwnerFilter = (ESamplerFilter)UDeviceProfileManager::Get().GetActiveProfile()->GetTextureLODSettings()->GetSamplerFilter(&Owner);
-#else
-	ESamplerFilter OwnerFilter = SF_Bilinear;
-#endif
+	ESamplerFilter OwnerFilter = bEnableGenMips ? (ESamplerFilter)UDeviceProfileManager::Get().GetActiveProfile()->GetTextureLODSettings()->GetSamplerFilter(&Owner) : SF_Bilinear;
 
 	if (CurrentSamplerFilter != OwnerFilter)
 	{
@@ -451,7 +445,7 @@ void FMediaTextureResource::InitDynamicRHI()
 	check(OutputTarget.IsValid());
 
 	// Register "external texture" parameters if the platform does not support them (and hence the player does not set them)
-	if (!GSupportsImageExternal)
+	if (!bUsesImageExternal)
 	{
 		FTextureRHIRef VideoTexture = (FTextureRHIRef)Owner.TextureReference.TextureReferenceRHI;
 		FExternalTextureRegistry::Get().RegisterExternalTexture(InitialTextureGuid, VideoTexture, SamplerStateRHI);
@@ -720,7 +714,7 @@ void FMediaTextureResource::ConvertSample(const TSharedPtr<IMediaTextureSample,
 }
 
 
-void FMediaTextureResource::CopySample(const TSharedPtr<IMediaTextureSample, ESPMode::ThreadSafe>& Sample, const FLinearColor& ClearColor, bool SrgbOutput, uint8 InNumMips)
+void FMediaTextureResource::CopySample(const TSharedPtr<IMediaTextureSample, ESPMode::ThreadSafe>& Sample, const FLinearColor& ClearColor, bool SrgbOutput, uint8 InNumMips, const FGuid & TextureGUID)
 {
 	FRHITexture* SampleTexture = Sample->GetTexture();
 	FRHITexture2D* SampleTexture2D = (SampleTexture != nullptr) ? SampleTexture->GetTexture2D() : nullptr;
@@ -731,38 +725,101 @@ void FMediaTextureResource::CopySample(const TSharedPtr<IMediaTextureSample, ESP
 
 	if (SampleTexture2D != nullptr)
 	{
-		if (InNumMips == 1)
-		{
-			// Use sample's texture as the new render target.
-		if (TextureRHI != SampleTexture2D)
-		{
-			UpdateTextureReference(SampleTexture2D);
+		// Use sample's texture as the new render target - no copy
+			if (TextureRHI != SampleTexture2D)
+			{
+				UpdateTextureReference(SampleTexture2D);
 
-			OutputTarget.SafeRelease();
-			UpdateResourceSize();
+				OutputTarget.SafeRelease();
 		}
-
-	}
-	else
-	{
+		else
+		{
 			// Texture to receive texture from sample
 			CreateOutputRenderTarget(Sample->GetOutputDim(), MediaTextureResourceHelpers::GetPixelFormat(Sample), SrgbOutput, ClearColor, InNumMips);
 
 			// Copy data into the output texture to able to add mips later on
 			FRHICommandListExecutor::GetImmediateCommandList().CopyTexture(SampleTexture2D, OutputTarget, FRHICopyTextureInfo());
 		}
-
-		Cleared = false;
-		}
+	}
 	else
-		{
+	{
 		// Texture to receive precisely only output pixels via CPU copy
 		CreateOutputRenderTarget(Sample->GetDim(), MediaTextureResourceHelpers::GetPixelFormat(Sample), SrgbOutput, ClearColor, InNumMips);
 
-		// copy sample data to output render target
+		if (bEnableGenMips && bUsesImageExternal)
+		{
+			// we are using an external texture to read from
+			CopyFromExternalTexture(Sample, TextureGUID);
+		}
+		else
+		{
+			// copy sample data (from CPU mem) to output render target
 		FUpdateTextureRegion2D Region(0, 0, 0, 0, Sample->GetDim().X, Sample->GetDim().Y);
 		RHIUpdateTexture2D(RenderTargetTextureRHI.GetReference(), 0, Region, Sample->GetStride(), (uint8*)Sample->GetBuffer());
 	}
+	}
+
+	Cleared = false;
+}
+
+
+void FMediaTextureResource::CopyFromExternalTexture(const TSharedPtr <IMediaTextureSample, ESPMode::ThreadSafe>& Sample, const FGuid & TextureGUID)
+{
+	FRHICommandListImmediate& CommandList = FRHICommandListExecutor::GetImmediateCommandList();
+	{
+		FTextureRHIRef SampleTexture;
+		FSamplerStateRHIRef SamplerState;
+		if (!FExternalTextureRegistry::Get().GetExternalTexture(nullptr, TextureGUID, SampleTexture, SamplerState))
+		{
+			return;
+		}
+
+		FLinearColor Offset, ScaleRotation;
+		FExternalTextureRegistry::Get().GetExternalTextureCoordinateOffset(TextureGUID, Offset);
+		FExternalTextureRegistry::Get().GetExternalTextureCoordinateScaleRotation(TextureGUID, ScaleRotation);
+
+		SCOPED_DRAW_EVENT(CommandList, FMediaTextureResource_ConvertExternalTexture);
+		SCOPED_GPU_STAT(CommandList, MediaTextureResource);
+
+		FGraphicsPipelineStateInitializer GraphicsPSOInit;
+		FRHITexture* RenderTarget = RenderTargetTextureRHI.GetReference();
+
+		FRHIRenderPassInfo RPInfo(RenderTarget, ERenderTargetActions::DontLoad_Store);
+		CommandList.BeginRenderPass(RPInfo, TEXT("ConvertMedia_ExternalTexture"));
+		{
+			const FIntPoint OutputDim = Sample->GetOutputDim();
+
+			CommandList.ApplyCachedRenderTargets(GraphicsPSOInit);
+			CommandList.SetViewport(0, 0, 0.0f, OutputDim.X, OutputDim.Y, 1.0f);
+			
+			GraphicsPSOInit.DepthStencilState = TStaticDepthStencilState<false, CF_Always>::GetRHI();
+			GraphicsPSOInit.RasterizerState = TStaticRasterizerState<>::GetRHI();
+			GraphicsPSOInit.BlendState = TStaticBlendStateWriteMask<CW_RGBA, CW_NONE, CW_NONE, CW_NONE, CW_NONE, CW_NONE, CW_NONE, CW_NONE>::GetRHI();
+			GraphicsPSOInit.PrimitiveType = PT_TriangleStrip;
+
+			// configure media shaders
+			auto ShaderMap = GetGlobalShaderMap(GMaxRHIFeatureLevel);
+			TShaderMapRef<FMediaShadersVS> VertexShader(ShaderMap);
+
+			GraphicsPSOInit.BoundShaderState.VertexDeclarationRHI = GMediaVertexDeclaration.VertexDeclarationRHI;
+			GraphicsPSOInit.BoundShaderState.VertexShaderRHI = GETSAFERHISHADER_VERTEX(*VertexShader);
+
+			TShaderMapRef<FReadTextureExternalPS> CopyShader(ShaderMap);
+			GraphicsPSOInit.BoundShaderState.PixelShaderRHI = GETSAFERHISHADER_PIXEL(*CopyShader);
+			SetGraphicsPipelineState(CommandList, GraphicsPSOInit);
+			CopyShader->SetParameters(CommandList, SampleTexture, SamplerState, ScaleRotation, Offset);
+
+			// draw full size quad into render target
+			FVertexBufferRHIRef VertexBuffer = CreateTempMediaVertexBuffer();
+			CommandList.SetStreamSource(0, VertexBuffer, 0);
+			// set viewport to RT size
+			CommandList.SetViewport(0, 0, 0.0f, OutputDim.X, OutputDim.Y, 1.0f);
+
+			CommandList.DrawPrimitive(0, 2, 1);
+		}
+		CommandList.EndRenderPass();
+		CommandList.TransitionResource(EResourceTransitionAccess::EReadable, RenderTargetTextureRHI);
+	}
 }
 
 
diff --git a/Engine/Source/Runtime/MediaAssets/Private/Misc/MediaTextureResource.h b/Engine/Source/Runtime/MediaAssets/Private/Misc/MediaTextureResource.h
--- a/Engine/Source/Runtime/MediaAssets/Private/Misc/MediaTextureResource.h
+++ b/Engine/Source/Runtime/MediaAssets/Private/Misc/MediaTextureResource.h
@@ -39,9 +39,10 @@ public:
 	 * @param InOWnerSize Reference to the size in bytes of the texture that owns this resource (will be updated by resource).
 	 * @param InClearColor The initial clear color.
 	 * @param InTextureGuid The initial external texture GUID.
+	 * @param bEnableGenMips If true mips generation will be enabled (possibly optimizing for NumMips == 1 case)
 	 * @param InNumMips The initial number of mips to be generated for the output texture
 	 */
-	FMediaTextureResource(UMediaTexture& InOwner, FIntPoint& InOwnerDim, SIZE_T& InOwnerSize, FLinearColor InClearColor, FGuid InTextureGuid, uint8 InNumMips);
+	FMediaTextureResource(UMediaTexture& InOwner, FIntPoint& InOwnerDim, SIZE_T& InOwnerSize, FLinearColor InClearColor, FGuid InTextureGuid, bool bEnableGenMips, uint8 InNumMips);
 
 	/** Virtual destructor. */
 	virtual ~FMediaTextureResource() { }
@@ -135,7 +136,7 @@ protected:
 	 * @param Number of mips
 	 * @see ConvertSample
 	 */
-	void CopySample(const TSharedPtr<IMediaTextureSample, ESPMode::ThreadSafe>& Sample, const FLinearColor& ClearColor, bool SrgbOutput, uint8 InNumMips);
+	void CopySample(const TSharedPtr<IMediaTextureSample, ESPMode::ThreadSafe>& Sample, const FLinearColor& ClearColor, bool SrgbOutput, uint8 InNumMips, const FGuid & TextureGUID);
 
 	/** Calculates the current resource size and notifies the owner texture. */
 	void UpdateResourceSize();
@@ -162,6 +163,9 @@ protected:
 	/** Setup sampler state from owner's settings as needed */
 	void SetupSampler();
 
+	/** Copy to local buffer from external texture */
+	void CopyFromExternalTexture(const TSharedPtr <IMediaTextureSample, ESPMode::ThreadSafe>& Sample, const FGuid & TextureGUID);
+
 private:
 
 	// Class to maintain a reference to a IMediaTextureSample instance as long  as needed by RHI
@@ -176,6 +180,8 @@ private:
 		TSharedPtr<IMediaTextureSample, ESPMode::ThreadSafe> MediaSample;
 	};
 
+	/** Platform uses GL/ES ImageExternal */
+	bool bUsesImageExternal;
 
 	/** Whether the texture has been cleared. */
 	bool Cleared;
@@ -201,6 +207,9 @@ private:
 	/** Reference to the owner's texture size field. */
 	SIZE_T& OwnerSize;
 
+	/** Enable mips generation */
+	bool bEnableGenMips;
+
 	/** Current number of mips to be generated as output */
 	uint8 CurrentNumMips;
 
diff --git a/Engine/Source/Runtime/MediaAssets/Public/MediaTexture.h b/Engine/Source/Runtime/MediaAssets/Public/MediaTexture.h
--- a/Engine/Source/Runtime/MediaAssets/Public/MediaTexture.h
+++ b/Engine/Source/Runtime/MediaAssets/Public/MediaTexture.h
@@ -47,6 +47,10 @@ class MEDIAASSETS_API UMediaTexture
 	UPROPERTY(EditAnywhere, BlueprintReadWrite, Category="MediaTexture")
 	FLinearColor ClearColor;
 
+	/** Basic enablement for mip generation. */
+	UPROPERTY(EditAnywhere, BlueprintReadOnly, Category = "MediaTexture", meta = (DisplayName = "Enable Mipmap generation"))
+	bool EnableGenMips;
+
 	/** The number of mips to use (default = 1). */
 	UPROPERTY(EditAnywhere, BlueprintReadWrite, Category = "MediaTexture", meta=(DisplayName="Total number of Mipmaps to output"))
 	uint8 NumMips;
diff --git a/Engine/Source/Runtime/RenderCore/Private/MediaShaders.cpp b/Engine/Source/Runtime/RenderCore/Private/MediaShaders.cpp
--- a/Engine/Source/Runtime/RenderCore/Private/MediaShaders.cpp
+++ b/Engine/Source/Runtime/RenderCore/Private/MediaShaders.cpp
@@ -562,3 +562,32 @@ void FSetAlphaOnePS::SetParameters(FRHICommandList& CommandList, TRefCountPtr<FR
 	TUniformBufferRef<FSetAlphaOneUB> Data = TUniformBufferRef<FSetAlphaOneUB>::CreateUniformBufferImmediate(UB, UniformBuffer_SingleFrame);
 	SetUniformBufferParameter(CommandList, GetPixelShader(), GetUniformBufferParameter<FSetAlphaOneUB>(), Data);
 }
+
+
+/* FReadTextureExternalPS shader
+ *****************************************************************************/
+
+BEGIN_GLOBAL_SHADER_PARAMETER_STRUCT(FReadTextureExternalUB, )
+SHADER_PARAMETER_TEXTURE(TextureExternal, Texture)
+SHADER_PARAMETER_SAMPLER(SamplerState, SamplerP)
+SHADER_PARAMETER(FLinearColor, ScaleRotation)
+SHADER_PARAMETER(FVector2D, Offset)
+END_GLOBAL_SHADER_PARAMETER_STRUCT()
+
+IMPLEMENT_GLOBAL_SHADER_PARAMETER_STRUCT(FReadTextureExternalUB, "ReadTextureExternalUB");
+IMPLEMENT_SHADER_TYPE(, FReadTextureExternalPS, TEXT("/Engine/Private/MediaShaders.usf"), TEXT("ReadTextureExternalPS"), SF_Pixel);
+
+
+void FReadTextureExternalPS::SetParameters(FRHICommandList& CommandList, FTextureRHIRef TextureExt, FSamplerStateRHIRef SamplerState, const FLinearColor & ScaleRotation, const FLinearColor & Offset)
+{
+	FReadTextureExternalUB UB;
+	{
+		UB.SamplerP = SamplerState;
+		UB.Texture = TextureExt;
+		UB.ScaleRotation = ScaleRotation;
+		UB.Offset = FVector2D(Offset.R, Offset.G);
+	}
+
+	TUniformBufferRef<FReadTextureExternalUB> Data = TUniformBufferRef<FReadTextureExternalUB>::CreateUniformBufferImmediate(UB, UniformBuffer_SingleFrame);
+	SetUniformBufferParameter(CommandList, GetPixelShader(), GetUniformBufferParameter<FReadTextureExternalUB>(), Data);
+}
diff --git a/Engine/Source/Runtime/RenderCore/Public/MediaShaders.h b/Engine/Source/Runtime/RenderCore/Public/MediaShaders.h
--- a/Engine/Source/Runtime/RenderCore/Public/MediaShaders.h
+++ b/Engine/Source/Runtime/RenderCore/Public/MediaShaders.h
@@ -664,3 +664,28 @@ public:
 
 	RENDERCORE_API void SetParameters(FRHICommandList& RHICmdList, TRefCountPtr<FRHITexture2D> RGBATexture);
 };
+
+
+/**
+ * Pixel shader to read from TextureExternal source
+ */
+class FReadTextureExternalPS
+	: public FGlobalShader
+{
+	DECLARE_EXPORTED_SHADER_TYPE(FReadTextureExternalPS, Global, RENDERCORE_API);
+
+public:
+
+	static bool ShouldCompilePermutation(const FGlobalShaderPermutationParameters& Parameters)
+	{
+		return IsFeatureLevelSupported(Parameters.Platform, ERHIFeatureLevel::ES2);
+	}
+
+	FReadTextureExternalPS() { }
+
+	FReadTextureExternalPS(const ShaderMetaType::CompiledShaderInitializerType& Initializer)
+		: FGlobalShader(Initializer)
+	{ }
+
+	RENDERCORE_API void SetParameters(FRHICommandList& RHICmdList, FTextureRHIRef TextureExt, FSamplerStateRHIRef SamplerState, const FLinearColor & ScaleRotation, const FLinearColor & Offset);
+};
