Adjustments and cleanup in resource state transtions (esp. re: Vulkan)
#rnx

diff --git a/Engine/Source/Runtime/MediaAssets/Private/Misc/MediaTextureResource.cpp b/Engine/Source/Runtime/MediaAssets/Private/Misc/MediaTextureResource.cpp
--- a/Engine/Source/Runtime/MediaAssets/Private/Misc/MediaTextureResource.cpp
+++ b/Engine/Source/Runtime/MediaAssets/Private/Misc/MediaTextureResource.cpp
@@ -307,8 +307,12 @@ void FMediaTextureResource::Render(const FRenderParams& Params)
 
 				FMemMark MemMark(FMemStack::Get());
 				FRHICommandListImmediate & RHICmdList = FRHICommandListExecutor::GetImmediateCommandList();
-
+#if PLATFORM_ANDROID
+				// Vulkan does not implement RWBarrier & ES does not do anything specific anyways
+				RHICmdList.TransitionResource(EResourceTransitionAccess::EWritable, OutputTarget);
+#else
 				RHICmdList.TransitionResource(EResourceTransitionAccess::ERWBarrier, OutputTarget);
+#endif
 				FGenerateMips::Execute(RHICmdList, OutputTarget, FGenerateMipsParams{ SF_Bilinear, AM_Clamp, AM_Clamp, AM_Clamp });
 				RHICmdList.TransitionResource(EResourceTransitionAccess::EReadable, OutputTarget);
 			}
@@ -483,8 +487,6 @@ void FMediaTextureResource::ClearTexture(const FLinearColor& ClearColor, bool Sr
 		FRHIRenderPassInfo RPInfo(RenderTargetTextureRHI, ERenderTargetActions::Clear_Store);
 		CommandList.BeginRenderPass(RPInfo, TEXT("ClearTexture"));
 		CommandList.EndRenderPass();
-		CommandList.SetViewport(0, 0, 0.0f, RenderTargetTextureRHI->GetSizeX(), RenderTargetTextureRHI->GetSizeY(), 1.0f);
-		CommandList.TransitionResource(EResourceTransitionAccess::EReadable, RenderTargetTextureRHI);
 	}
 
 	Cleared = true;
@@ -707,7 +709,6 @@ void FMediaTextureResource::ConvertSample(const TSharedPtr<IMediaTextureSample,
 			CommandList.DrawPrimitive(0, 2, 1);
 		}
 		CommandList.EndRenderPass();
-		CommandList.TransitionResource(EResourceTransitionAccess::EReadable, RenderTargetTextureRHI);
 	}
 
 	Cleared = false;
@@ -746,14 +747,14 @@ void FMediaTextureResource::CopySample(const TSharedPtr<IMediaTextureSample, ESP
 		// Texture to receive precisely only output pixels via CPU copy
 		CreateOutputRenderTarget(Sample->GetDim(), MediaTextureResourceHelpers::GetPixelFormat(Sample), SrgbOutput, ClearColor, InNumMips);
 
-		if (bEnableGenMips && bUsesImageExternal)
+		if (bEnableGenMips && bUsesImageExternal && !Sample->GetBuffer())
 		{
-			// we are using an external texture to read from
+			// We are using an external texture to read from
 			CopyFromExternalTexture(Sample, TextureGUID);
 		}
 		else
 		{
-			// copy sample data (from CPU mem) to output render target
+			// Copy sample data (from CPU mem) to output render target
 		FUpdateTextureRegion2D Region(0, 0, 0, 0, Sample->GetDim().X, Sample->GetDim().Y);
 		RHIUpdateTexture2D(RenderTargetTextureRHI.GetReference(), 0, Region, Sample->GetStride(), (uint8*)Sample->GetBuffer());
 	}
@@ -771,6 +772,10 @@ void FMediaTextureResource::CopyFromExternalTexture(const TSharedPtr <IMediaText
 		FSamplerStateRHIRef SamplerState;
 		if (!FExternalTextureRegistry::Get().GetExternalTexture(nullptr, TextureGUID, SampleTexture, SamplerState))
 		{
+			// This should never happen: we could not find the external texture data. Still, if it does we clear the output...
+			FRHIRenderPassInfo RPInfo(RenderTargetTextureRHI, ERenderTargetActions::Clear_Store);
+			CommandList.BeginRenderPass(RPInfo, TEXT("ClearTexture"));
+			CommandList.EndRenderPass();
 			return;
 		}
 
@@ -818,7 +823,6 @@ void FMediaTextureResource::CopyFromExternalTexture(const TSharedPtr <IMediaText
 			CommandList.DrawPrimitive(0, 2, 1);
 		}
 		CommandList.EndRenderPass();
-		CommandList.TransitionResource(EResourceTransitionAccess::EReadable, RenderTargetTextureRHI);
 	}
 }
 
