Reworked code to add/remove audio sample sink based on activate state of the UMediaSoundComponent. This fix prevents audio decoders with a fixed number of pre-allocated audio buffers to stall. If the UMediaSoundComponent stops playing the samples from the created audio sink are not consumed and/or released. The samples stay in the audio sink preventing the audio decoder from generating new samples as all buffers are occupied (waiting to be released from the audio sink by playing for example)

#rnx

diff --git a/Engine/Source/Runtime/MediaAssets/Private/Assets/MediaSoundComponent.cpp b/Engine/Source/Runtime/MediaAssets/Private/Assets/MediaSoundComponent.cpp
--- a/Engine/Source/Runtime/MediaAssets/Private/Assets/MediaSoundComponent.cpp
+++ b/Engine/Source/Runtime/MediaAssets/Private/Assets/MediaSoundComponent.cpp
@@ -142,17 +142,37 @@ void UMediaSoundComponent::UpdatePlayer()
 	// create a new sample queue if the player changed
 	TSharedRef<FMediaPlayerFacade, ESPMode::ThreadSafe> PlayerFacade = CurrentPlayerPtr->GetPlayerFacade();
 
+	// We have some audio decoders which are running with a limited amount of pre-allocated audio sample packets. 
+	// When the audio packets are not consumed in the UMediaSoundComponent::OnGenerateAudio method below, these packets are not 
+	// returned to the decooder which then cannot produce more audio samples. 
+	//
+	// The UMediaSoundComponent::OnGenerateAudio is only called when our parent USynthComponent it active and
+	// this is conrolled by USynthComponent::Start() and USynthComponent::Stop(). We are tracking a state change here.
 	if (PlayerFacade != CurrentPlayerFacade)
 	{
-		const auto NewSampleQueue = MakeShared<FMediaAudioSampleQueue, ESPMode::ThreadSafe>();
-		PlayerFacade->AddAudioSampleSink(NewSampleQueue);
+		if (IsActive())
+		{
+			const auto NewSampleQueue = MakeShared<FMediaAudioSampleQueue, ESPMode::ThreadSafe>();
+			PlayerFacade->AddAudioSampleSink(NewSampleQueue);
+			{
+				FScopeLock Lock(&CriticalSection);
+				SampleQueue = NewSampleQueue;
+				FrameSyncOffset = 0;
+			}
+			CurrentPlayerFacade = PlayerFacade;
+		}
+	}
+	else
+	{
+		// Here, we have a CurrentPlayerFacade set which means are also have a valid FMediaAudioSampleQueue set
+		// We need to check for deactivation as it seems there is not callback scheduled when USynthComponent::Stop() is called.
+		if(!IsActive())
 		{
 			FScopeLock Lock(&CriticalSection);
-			SampleQueue = NewSampleQueue;
+			SampleQueue.Reset();
 			FrameSyncOffset = 0;
+			CurrentPlayerFacade.Reset();
 		}
-
-		CurrentPlayerFacade = PlayerFacade;
 	}
 
 	//
@@ -215,6 +235,17 @@ void UMediaSoundComponent::OnRegister()
 #endif
 }
 
+void UMediaSoundComponent::OnUnregister()
+{
+	{
+		FScopeLock Lock(&CriticalSection);
+		SampleQueue.Reset();
+		FrameSyncOffset = 0;
+	}
+	CurrentPlayerFacade.Reset();
+	Super::OnUnregister();
+}
+
 
 void UMediaSoundComponent::TickComponent(float DeltaTime, enum ELevelTick TickType, FActorComponentTickFunction *ThisTickFunction)
 {
@@ -243,8 +274,13 @@ void UMediaSoundComponent::Deactivate()
 	if (!ShouldActivate())
 	{
 		SetComponentTickEnabled(false);
+		{
+			FScopeLock Lock(&CriticalSection);
+			SampleQueue.Reset();
+			FrameSyncOffset = 0;
+		}
+		CurrentPlayerFacade.Reset();
 	}
-
 	Super::Deactivate();
 }
 
diff --git a/Engine/Source/Runtime/MediaAssets/Public/MediaSoundComponent.h b/Engine/Source/Runtime/MediaAssets/Public/MediaSoundComponent.h
--- a/Engine/Source/Runtime/MediaAssets/Public/MediaSoundComponent.h
+++ b/Engine/Source/Runtime/MediaAssets/Public/MediaSoundComponent.h
@@ -194,6 +194,7 @@ public:
 	//~ UActorComponent interface
 
 	virtual void OnRegister() override;
+	virtual void OnUnregister() override;
 	virtual void TickComponent(float DeltaTime, enum ELevelTick TickType, FActorComponentTickFunction *ThisTickFunction) override;
 
 public:
