MediaTextureResources fixup for missing texture state transitions (failed PS4 clear) & adding explicit caching of FGenerateMips state setup (memory leak fix)
#rnx

diff --git a/Engine/Source/Runtime/MediaAssets/Private/Misc/MediaTextureResource.cpp b/Engine/Source/Runtime/MediaAssets/Private/Misc/MediaTextureResource.cpp
--- a/Engine/Source/Runtime/MediaAssets/Private/Misc/MediaTextureResource.cpp
+++ b/Engine/Source/Runtime/MediaAssets/Private/Misc/MediaTextureResource.cpp
@@ -313,7 +313,7 @@ void FMediaTextureResource::Render(const FRenderParams& Params)
 #else
 				RHICmdList.TransitionResource(EResourceTransitionAccess::ERWBarrier, OutputTarget);
 #endif
-				FGenerateMips::Execute(RHICmdList, OutputTarget, FGenerateMipsParams{ SF_Bilinear, AM_Clamp, AM_Clamp, AM_Clamp });
+				FGenerateMips::Execute(RHICmdList, OutputTarget, FGenerateMipsParams{ SF_Bilinear, AM_Clamp, AM_Clamp, AM_Clamp }, &CachedMipsGenParams, true);
 				RHICmdList.TransitionResource(EResourceTransitionAccess::EReadable, OutputTarget);
 			}
 
@@ -461,6 +461,8 @@ void FMediaTextureResource::ReleaseDynamicRHI()
 {
 	Cleared = false;
 
+	CachedMipsGenParams.Reset();
+
 	InputTarget.SafeRelease();
 	OutputTarget.SafeRelease();
 	RenderTargetTextureRHI.SafeRelease();
@@ -487,6 +489,7 @@ void FMediaTextureResource::ClearTexture(const FLinearColor& ClearColor, bool Sr
 		FRHIRenderPassInfo RPInfo(RenderTargetTextureRHI, ERenderTargetActions::Clear_Store);
 		CommandList.BeginRenderPass(RPInfo, TEXT("ClearTexture"));
 		CommandList.EndRenderPass();
+		CommandList.TransitionResource(EResourceTransitionAccess::EReadable, RenderTargetTextureRHI);
 	}
 
 	Cleared = true;
@@ -709,6 +712,7 @@ void FMediaTextureResource::ConvertSample(const TSharedPtr<IMediaTextureSample,
 			CommandList.DrawPrimitive(0, 2, 1);
 		}
 		CommandList.EndRenderPass();
+		CommandList.TransitionResource(EResourceTransitionAccess::EReadable, RenderTargetTextureRHI);
 	}
 
 	Cleared = false;
@@ -727,11 +731,12 @@ void FMediaTextureResource::CopySample(const TSharedPtr<IMediaTextureSample, ESP
 	if (SampleTexture2D != nullptr)
 	{
 		// Use sample's texture as the new render target - no copy
-			if (TextureRHI != SampleTexture2D)
-			{
-				UpdateTextureReference(SampleTexture2D);
+		if (TextureRHI != SampleTexture2D)
+		{
+			UpdateTextureReference(SampleTexture2D);
 
-				OutputTarget.SafeRelease();
+			CachedMipsGenParams.Reset();
+			OutputTarget.SafeRelease();
 		}
 		else
 		{
@@ -755,9 +760,9 @@ void FMediaTextureResource::CopySample(const TSharedPtr<IMediaTextureSample, ESP
 		else
 		{
 			// Copy sample data (from CPU mem) to output render target
-		FUpdateTextureRegion2D Region(0, 0, 0, 0, Sample->GetDim().X, Sample->GetDim().Y);
-		RHIUpdateTexture2D(RenderTargetTextureRHI.GetReference(), 0, Region, Sample->GetStride(), (uint8*)Sample->GetBuffer());
-	}
+			FUpdateTextureRegion2D Region(0, 0, 0, 0, Sample->GetDim().X, Sample->GetDim().Y);
+			RHIUpdateTexture2D(RenderTargetTextureRHI.GetReference(), 0, Region, Sample->GetStride(), (uint8*)Sample->GetBuffer());
+		}
 	}
 
 	Cleared = false;
@@ -776,6 +781,7 @@ void FMediaTextureResource::CopyFromExternalTexture(const TSharedPtr <IMediaText
 			FRHIRenderPassInfo RPInfo(RenderTargetTextureRHI, ERenderTargetActions::Clear_Store);
 			CommandList.BeginRenderPass(RPInfo, TEXT("ClearTexture"));
 			CommandList.EndRenderPass();
+			CommandList.TransitionResource(EResourceTransitionAccess::EReadable, RenderTargetTextureRHI);
 			return;
 		}
 
@@ -823,6 +829,7 @@ void FMediaTextureResource::CopyFromExternalTexture(const TSharedPtr <IMediaText
 			CommandList.DrawPrimitive(0, 2, 1);
 		}
 		CommandList.EndRenderPass();
+		CommandList.TransitionResource(EResourceTransitionAccess::EReadable, RenderTargetTextureRHI);
 	}
 }
 
@@ -881,6 +888,8 @@ void FMediaTextureResource::CreateOutputRenderTarget(const FIntPoint & InDim, EP
 	{
 		TRefCountPtr<FRHITexture2D> DummyTexture2DRHI;
 
+		CachedMipsGenParams.Reset();
+
 		FRHIResourceCreateInfo CreateInfo = {
 			FClearValueBinding(InClearColor)
 		};
diff --git a/Engine/Source/Runtime/MediaAssets/Private/Misc/MediaTextureResource.h b/Engine/Source/Runtime/MediaAssets/Private/Misc/MediaTextureResource.h
--- a/Engine/Source/Runtime/MediaAssets/Private/Misc/MediaTextureResource.h
+++ b/Engine/Source/Runtime/MediaAssets/Private/Misc/MediaTextureResource.h
@@ -17,6 +17,7 @@ class FMediaPlayerFacade;
 class IMediaPlayer;
 class IMediaTextureSample;
 class UMediaTexture;
+struct FGenerateMipsStruct;
 
 enum class EMediaTextureSinkFormat;
 enum class EMediaTextureSinkMode;
@@ -221,4 +222,7 @@ private:
 
 	/** cached media sample to postpone releasing it until the next sample rendering as it can get overwritten due to asynchronous rendering */
 	TRefCountPtr<FTextureSampleKeeper> CurrentSample;
+
+	/** cached params etc. for use with mip generator */
+	TSharedPtr<FGenerateMipsStruct> CachedMipsGenParams;
 };
