added play command that can also seek #jira none

diff --git a/Engine/Source/Runtime/MediaAssets/Private/Assets/MediaPlayer.cpp b/Engine/Source/Runtime/MediaAssets/Private/Assets/MediaPlayer.cpp
--- a/Engine/Source/Runtime/MediaAssets/Private/Assets/MediaPlayer.cpp
+++ b/Engine/Source/Runtime/MediaAssets/Private/Assets/MediaPlayer.cpp
@@ -533,6 +533,18 @@ bool UMediaPlayer::Play()
 	return PlayerFacade->SetRate(1.0f);
 }
 
+void UMediaPlayer::PlayAndSeek()
+{
+	PlayOnNext = false;
+	if (Play())
+	{
+		if (PlayerFacade->ActivePlayerOptions.IsSet() && !PlayerFacade->ActivePlayerOptions->SeekTime.IsZero() && SupportsSeeking())
+		{
+			Seek(PlayerFacade->ActivePlayerOptions->SeekTime);
+		}
+	}
+}
+
 
 bool UMediaPlayer::Previous()
 {
@@ -891,14 +903,7 @@ void UMediaPlayer::HandlePlayerMediaEvent(EMediaEvent Event)
 
 		if (bPlayOnOpen)
 		{
-			PlayOnNext = false;
-			if (Play())
-			{
-				if (PlayerFacade->ActivePlayerOptions.IsSet() && !PlayerFacade->ActivePlayerOptions->SeekTime.IsZero() && SupportsSeeking())
-				{
-					Seek(PlayerFacade->ActivePlayerOptions->SeekTime);
-				}
-			}
+			PlayAndSeek();
 		}
 		break;
 
diff --git a/Engine/Source/Runtime/MediaAssets/Public/MediaPlayer.h b/Engine/Source/Runtime/MediaAssets/Public/MediaPlayer.h
--- a/Engine/Source/Runtime/MediaAssets/Public/MediaPlayer.h
+++ b/Engine/Source/Runtime/MediaAssets/Public/MediaPlayer.h
@@ -645,6 +645,12 @@ public:
 	UFUNCTION(BlueprintCallable, Category="Media|MediaPlayer")
 	bool Play();
 
+	/**
+	 * Starts playback from the media opened event, but can be used elsewhere.
+	 */
+	UFUNCTION(BlueprintCallable, Category="Media|MediaPlayer")
+	void PlayAndSeek();
+
 	/**
 	 * Open the previous item in the current play list.
 	 *
