MediaFramework: Fixed wrong sample being used sometimes in MediaTextureResource.

diff --git a/Engine/Source/Runtime/MediaAssets/Private/Misc/MediaTextureResource.cpp b/Engine/Source/Runtime/MediaAssets/Private/Misc/MediaTextureResource.cpp
--- a/Engine/Source/Runtime/MediaAssets/Private/Misc/MediaTextureResource.cpp
+++ b/Engine/Source/Runtime/MediaAssets/Private/Misc/MediaTextureResource.cpp
@@ -186,12 +186,13 @@ void FMediaTextureResource::Render(const FRenderParams& Params)
 	{
 		// get the most current sample to be rendered
 		TSharedPtr<IMediaTextureSample, ESPMode::ThreadSafe> Sample;
+		TSharedPtr<IMediaTextureSample, ESPMode::ThreadSafe> PeekSample;
 		bool UseSample = false;
 		
-		while (SampleSource->Peek(Sample) && Sample.IsValid())
+		while (SampleSource->Peek(PeekSample) && PeekSample.IsValid())
 		{
-			const FTimespan StartTime = Sample->GetTime();
-			const FTimespan EndTime = StartTime + Sample->GetDuration();
+			const FTimespan StartTime = PeekSample->GetTime();
+			const FTimespan EndTime = StartTime + PeekSample->GetDuration();
 
 			if ((Params.Rate >= 0.0f) && (Params.Time < StartTime))
 			{
