Integration of recent Electra player changes

#rnx

diff --git a/Engine/Source/Runtime/Media/Public/IMediaPlayer.h b/Engine/Source/Runtime/Media/Public/IMediaPlayer.h
--- a/Engine/Source/Runtime/Media/Public/IMediaPlayer.h
+++ b/Engine/Source/Runtime/Media/Public/IMediaPlayer.h
@@ -166,22 +166,6 @@ public:
 	 */
 	virtual bool Open(const TSharedRef<FArchive, ESPMode::ThreadSafe>& Archive, const FString& OriginalUrl, const IMediaOptions* Options) = 0;
 
-	/**
-	* Specifies how audio starving should be handled
-	* 
-	* This affects the starving behavior of the used UMediaSoundComponent component.
-	* When true, it will skip audio to keep the time as if no starving occurred.
-	* If false, it will simply wait and continue without skipping when it gets more audio.
-	*/
-	virtual bool RequiresAudioSyncAfterDropouts() const
-	{
-#if PLATFORM_PS4 || PLATFORM_XBOXONE
-		return true;
-#else
-		return false;
-#endif
-	}
-
 public:
 
 	//~ The following methods are optional
diff --git a/Engine/Source/Runtime/MediaAssets/Private/Assets/MediaSoundComponent.cpp b/Engine/Source/Runtime/MediaAssets/Private/Assets/MediaSoundComponent.cpp
--- a/Engine/Source/Runtime/MediaAssets/Private/Assets/MediaSoundComponent.cpp
+++ b/Engine/Source/Runtime/MediaAssets/Private/Assets/MediaSoundComponent.cpp
@@ -17,14 +17,6 @@
 #include "MediaPlayerFacade.h"
 
 
-static int32 SyncAudioAfterDropoutsCVar = 1;
-FAutoConsoleVariableRef CVarSyncAudioAfterDropouts(
-	TEXT("m.SyncAudioAfterDropouts"),
-	SyncAudioAfterDropoutsCVar,
-	TEXT("Skip over delayed contiguous audio samples.\n")
-	TEXT("0: Not Enabled, 1: Enabled"),
-	ECVF_Default);
-
 DECLARE_FLOAT_COUNTER_STAT(TEXT("MediaUtils MediaSoundComponent Sync"), STAT_MediaUtils_MediaSoundComponentSync, STATGROUP_Media);
 DECLARE_FLOAT_COUNTER_STAT(TEXT("MediaUtils MediaSoundComponent SampleTime"), STAT_MediaUtils_MediaSoundComponentSampleTime, STATGROUP_Media);
 DECLARE_DWORD_COUNTER_STAT(TEXT("MediaUtils MediaSoundComponent Queued"), STAT_Media_SoundCompQueued, STATGROUP_Media);
@@ -50,12 +42,10 @@ UMediaSoundComponent::UMediaSoundComponent(const FObjectInitializer& ObjectIniti
 	, CachedTime(FTimespan::Zero())
 	, RateAdjustment(1.0f)
 	, Resampler(new FMediaAudioResampler)
-	, FrameSyncOffset(0)
 	, LastPlaySampleTime(FTimespan::MinValue())
 	, EnvelopeFollowerAttackTime(10)
 	, EnvelopeFollowerReleaseTime(100)
 	, CurrentEnvelopeValue(0.0f)
-	, bSyncAudioAfterDropouts(false)
 	, bSpectralAnalysisEnabled(false)
 	, bEnvelopeFollowingEnabled(false)
 {
@@ -69,12 +59,6 @@ UMediaSoundComponent::UMediaSoundComponent(const FObjectInitializer& ObjectIniti
 #if WITH_EDITORONLY_DATA
 	bVisualizeComponent = true;
 #endif
-
-#if PLATFORM_PS4 || PLATFORM_XBOXONE
-	bSyncAudioAfterDropouts = true;
-#else
-	bSyncAudioAfterDropouts = false;
-#endif
 }
 
 
@@ -134,7 +118,6 @@ void UMediaSoundComponent::UpdatePlayer()
 
 		FScopeLock Lock(&CriticalSection);
 		SampleQueue.Reset();
-		FrameSyncOffset = 0;
 
 		return;
 	}
@@ -157,7 +140,6 @@ void UMediaSoundComponent::UpdatePlayer()
 			{
 				FScopeLock Lock(&CriticalSection);
 				SampleQueue = NewSampleQueue;
-				FrameSyncOffset = 0;
 			}
 			CurrentPlayerFacade = PlayerFacade;
 		}
@@ -170,24 +152,10 @@ void UMediaSoundComponent::UpdatePlayer()
 		{
 			FScopeLock Lock(&CriticalSection);
 			SampleQueue.Reset();
-			FrameSyncOffset = 0;
 			CurrentPlayerFacade.Reset();
 		}
 	}
 
-	//
-	// Some players might require bSyncAudioAfterDropouts to be forced to true or false
-	if (PlayerFacade->GetPlayer())
-	{
-		bool bNewAudioSyncAfterDropouts = PlayerFacade->GetPlayer()->RequiresAudioSyncAfterDropouts();
-		if (bNewAudioSyncAfterDropouts != bSyncAudioAfterDropouts)
-		{
-			FScopeLock Lock(&CriticalSection);
-			bSyncAudioAfterDropouts = bNewAudioSyncAfterDropouts;
-			FrameSyncOffset = 0;
-		}
-	}
-
 	// caching play rate and time for audio thread (eventual consistency is sufficient)
 	CachedRate = PlayerFacade->GetRate();
 	CachedTime = PlayerFacade->GetTime();
@@ -240,7 +208,6 @@ void UMediaSoundComponent::OnUnregister()
 	{
 		FScopeLock Lock(&CriticalSection);
 		SampleQueue.Reset();
-		FrameSyncOffset = 0;
 	}
 	CurrentPlayerFacade.Reset();
 	Super::OnUnregister();
@@ -277,7 +244,6 @@ void UMediaSoundComponent::Deactivate()
 		{
 			FScopeLock Lock(&CriticalSection);
 			SampleQueue.Reset();
-			FrameSyncOffset = 0;
 		}
 		CurrentPlayerFacade.Reset();
 	}
@@ -383,7 +349,6 @@ int32 UMediaSoundComponent::OnGenerateAudio(float* OutAudio, int32 NumSamples)
 	{
 		FScopeLock Lock(&CriticalSection);
 		PinnedSampleQueue = SampleQueue;
-		InitialSyncOffset = FrameSyncOffset;
 	}
 
 	if (PinnedSampleQueue.IsValid() && (CachedRate != 0.0f))
@@ -393,70 +358,6 @@ int32 UMediaSoundComponent::OnGenerateAudio(float* OutAudio, int32 NumSamples)
 
 		FTimespan OutTime = FTimespan::Zero();
 
-		if (bSyncAudioAfterDropouts && SyncAudioAfterDropoutsCVar)
-		{
-			int32 SyncOffset = InitialSyncOffset + (NumSamples / NumChannels);
-			while (SyncOffset > 0)
-			{
-				float* DestAudio = OutAudio;
-				int32 FramesRequested = NumSamples / NumChannels;
-
-				if (SyncOffset < FramesRequested)
-				{
-					// Handle final generate before audio resumes playback
-					// Move frames left to sync them with expected playback time
-					int32 FloatsMoved = (FramesRequested - SyncOffset) * NumChannels;
-					FMemory::Memmove(OutAudio, OutAudio + (SyncOffset * NumChannels), FloatsMoved * sizeof(float));
-					DestAudio = OutAudio + FloatsMoved;
-					FramesRequested = SyncOffset;
-				}
-
-				uint32 JumpFrame = MAX_uint32;
-				int32 FramesWritten = (int32)Resampler->Generate(DestAudio, OutTime, (uint32)FramesRequested, Rate, Time, *PinnedSampleQueue, JumpFrame);
-
-				if (JumpFrame != MAX_uint32)
-				{
-					UE_LOG(LogMediaAssets, Verbose, TEXT("Audio ( JUMP ) SyncOffset was: %d, OutTime: %s"), SyncOffset, *OutTime.ToString());
-					int32 JumpFramesRequested = FramesRequested - JumpFrame;
-					int32 JumpFramesWritten = FramesWritten - JumpFrame;
-					SyncOffset = JumpFramesRequested - JumpFramesWritten;
-				}
-				else
-				{
-					SyncOffset -= FramesWritten;
-				}
-
-				if (FramesWritten < FramesRequested)
-				{
-					if (FramesWritten > 0)
-					{
-						UE_LOG(LogMediaAssets, Verbose, TEXT("Audio partial generate, FramesWritten: %d"), FramesWritten);
-					}
-					// Source buffer is empty
-					break;
-				}
-			}
-
-			if (SyncOffset > 0)
-			{
-				UE_LOG(LogMediaAssets, Verbose, TEXT("Audio ( STARVED ) SyncOffset: %d, PlayerTime: %s, OutTime: %s"), SyncOffset, *Time.ToString(), *OutTime.ToString());
-				FMemory::Memzero(OutAudio, NumSamples * sizeof(float));
-			}
-			else if (SyncOffset < 0)
-			{
-				UE_LOG(LogMediaAssets, Verbose, TEXT("Audio ( DESYNCED ) SyncOffset: %d"), SyncOffset);
-			}
-
-			{
-				FScopeLock Lock(&CriticalSection);
-				// Commit only if another thread did not change value
-				if (InitialSyncOffset == FrameSyncOffset)
-				{
-					FrameSyncOffset = SyncOffset;
-				}
-			}
-		}
-		else
 		{
 			const uint32 FramesRequested = uint32(NumSamples / NumChannels);
 			uint32 JumpFrame = MAX_uint32;
@@ -535,12 +436,6 @@ int32 UMediaSoundComponent::OnGenerateAudio(float* OutAudio, int32 NumSamples)
 	{
 		Resampler->Flush();
 
-		if (bSyncAudioAfterDropouts && SyncAudioAfterDropoutsCVar)
-		{
-			FScopeLock Lock(&CriticalSection);
-			FrameSyncOffset = 0;
-		}
-
 		LastPlaySampleTime = FTimespan::MinValue();
 	}
 	return NumSamples;
diff --git a/Engine/Source/Runtime/MediaAssets/Public/MediaSoundComponent.h b/Engine/Source/Runtime/MediaAssets/Public/MediaSoundComponent.h
--- a/Engine/Source/Runtime/MediaAssets/Public/MediaSoundComponent.h
+++ b/Engine/Source/Runtime/MediaAssets/Public/MediaSoundComponent.h
@@ -268,10 +268,6 @@ private:
 	/** Audio sample queue. */
 	TSharedPtr<FMediaAudioSampleQueue, ESPMode::ThreadSafe> SampleQueue;
 
-	/** Handle SampleQueue running dry. Ensure audio resumes playback at correct position. */
-	int32 FrameSyncOffset;
-
-
 	/* Time of last sample played. */
 	TAtomic<FTimespan> LastPlaySampleTime;
 
@@ -294,12 +290,6 @@ private:
 	/** Scratch buffer to mix in source audio to from decoder */
 	Audio::AlignedFloatBuffer AudioScratchBuffer;
 
-	/**
-	 * Sync forward after input audio buffer runs dry due to a hitch or decoder not being able to keep up
-	 * Without this audio will resume playing exactly where it last left off (far behind current player time)
-	 */
-	bool bSyncAudioAfterDropouts;
-
 	/** Whether or not spectral analysis is enabled. */
 	bool bSpectralAnalysisEnabled;
 
