MediaFramework: hardening output logic for CPU-space/external sample processing against "empty" samples (no texture & no buffer) & fixing up a dangling render pass in case an illegal format is passed into the same logic
#rnx

diff --git a/Engine/Source/Runtime/MediaAssets/Private/Misc/MediaTextureResource.cpp b/Engine/Source/Runtime/MediaAssets/Private/Misc/MediaTextureResource.cpp
--- a/Engine/Source/Runtime/MediaAssets/Private/Misc/MediaTextureResource.cpp
+++ b/Engine/Source/Runtime/MediaAssets/Private/Misc/MediaTextureResource.cpp
@@ -842,7 +842,11 @@ void FMediaTextureResource::ReleaseDynamicRHI()
 			break;
 
 			default:
-				return; // unsupported format
+				{
+					// This should not happen in normal use: still - end the render pass to avoid any trouble with RHI
+					CommandList.EndRenderPass();
+					return; // unsupported format
+				}
 			}
 
 			// draw full size quad into render target
@@ -896,9 +900,23 @@ void FMediaTextureResource::CopySample(const TSharedPtr<IMediaTextureSample, ESP
 
 		// If we also have no source buffer and the platform generally would allow for use of external textures, we assume it is just that...
 		// (as long as the player actually produces (dummy) samples, this will enable mips support as well as auto conversion for "new style output" mode)
-		if (!Sample->GetBuffer() && GSupportsImageExternal)
+		if (!Sample->GetBuffer())
 		{
-			CopyFromExternalTexture(Sample, TextureGUID);
+			if (GSupportsImageExternal)
+			{
+				CopyFromExternalTexture(Sample, TextureGUID);
+			}
+			else
+			{
+				// We never should get here, but could should a player pass us a "valid" sample with neither texture or buffer based data in it (and we don't have ExternalTexture support)
+
+				// Just clear the texture so we don't show any random memory contents...
+				FRHICommandListImmediate& CommandList = FRHICommandListExecutor::GetImmediateCommandList();
+				FRHIRenderPassInfo RPInfo(RenderTargetTextureRHI, ERenderTargetActions::Clear_Store);
+				CommandList.BeginRenderPass(RPInfo, TEXT("ClearTexture"));
+				CommandList.EndRenderPass();
+				CommandList.TransitionResource(EResourceTransitionAccess::EReadable, RenderTargetTextureRHI);
+			}
 		}
 		else
 		{
