Make sure resource is in SRV after Texture update - fix resource state

diff --git a/Engine/Source/Runtime/MediaAssets/Private/Misc/MediaTextureResource.cpp b/Engine/Source/Runtime/MediaAssets/Private/Misc/MediaTextureResource.cpp
--- a/Engine/Source/Runtime/MediaAssets/Private/Misc/MediaTextureResource.cpp
+++ b/Engine/Source/Runtime/MediaAssets/Private/Misc/MediaTextureResource.cpp
@@ -939,6 +939,9 @@ void FMediaTextureResource::CopySample(const TSharedPtr<IMediaTextureSample, ESP
 			// Copy sample data (from CPU mem) to output render target
 			FUpdateTextureRegion2D Region(0, 0, 0, 0, Sample->GetDim().X, Sample->GetDim().Y);
 			RHIUpdateTexture2D(RenderTargetTextureRHI.GetReference(), 0, Region, Sample->GetStride(), (uint8*)Sample->GetBuffer());
+
+			// Make sure resource is in SRV mode again
+			FRHICommandListExecutor::GetImmediateCommandList().Transition(FRHITransitionInfo(RenderTargetTextureRHI.GetReference(), ERHIAccess::Unknown, ERHIAccess::SRVMask));
 		}
 	}
 
