Making is safe to call Close() on UMediaPlayer during GC hence getting rid of crashes observed as other objects are destroyed in a GC call Close() on the UMediaPlayer object during a GC

diff --git a/Engine/Source/Runtime/MediaAssets/Private/Assets/MediaPlayer.cpp b/Engine/Source/Runtime/MediaAssets/Private/Assets/MediaPlayer.cpp
--- a/Engine/Source/Runtime/MediaAssets/Private/Assets/MediaPlayer.cpp
+++ b/Engine/Source/Runtime/MediaAssets/Private/Assets/MediaPlayer.cpp
@@ -38,6 +38,7 @@ UMediaPlayer::UMediaPlayer(const FObjectInitializer& ObjectInitializer)
 	, PlayOnOpen(true)
 	, Shuffle(false)
 	, Loop(false)
+	, Playlist(nullptr)
 	, PlaylistIndex(INDEX_NONE)
 	, TimeDelay(FTimespan::Zero())
 	, HorizontalFieldOfView(90.0f)
@@ -55,7 +56,6 @@ UMediaPlayer::UMediaPlayer(const FObjectInitializer& ObjectInitializer)
 	{
 		PlayerFacade = MakeShareable(new FMediaPlayerFacade());
 		PlayerFacade->OnMediaEvent().AddUObject(this, &UMediaPlayer::HandlePlayerMediaEvent);
-		Playlist = NewObject<UMediaPlaylist>(GetTransientPackage(), NAME_None, RF_Transactional | RF_Transient);
 	}
 }
 
@@ -94,7 +94,17 @@ bool UMediaPlayer::CanPlayUrl(const FString& Url)
 	return PlayerFacade->CanPlayUrl(Url, GetDefault<UMediaSource>());
 }
 
-void UMediaPlayer::SetPlaylistInternal(UMediaPlaylist* InPlaylist)
+
+void UMediaPlayer::EnsurePlaylist() const
+{
+	if (!Playlist)
+	{
+		SetPlaylistInternal(NewObject<UMediaPlaylist>(GetTransientPackage(), NAME_None, RF_Transactional | RF_Transient));
+	}
+}
+
+
+void UMediaPlayer::SetPlaylistInternal(UMediaPlaylist* InPlaylist) const
 {
 	if (Playlist && Playlist != InPlaylist && Playlist->IsRooted())
 	{
@@ -118,11 +128,7 @@ void UMediaPlayer::Close()
 
 	PlayerFacade->Close();
 
-	Playlist = nullptr;
-	if (!HasAnyFlags(RF_ClassDefaultObject) && !GExitPurge)
-	{
-		SetPlaylistInternal(NewObject<UMediaPlaylist>(GetTransientPackage(), NAME_None, RF_Transactional | RF_Transient));
-	}
+	SetPlaylistInternal(nullptr);
 
 	PlaylistIndex = INDEX_NONE;
 	PlayOnNext = false;
@@ -384,7 +390,11 @@ bool UMediaPlayer::Next()
 {
 	UE_LOG(LogMediaAssets, Verbose, TEXT("%s.Next"), *GetFName().ToString());
 
-	check(Playlist != nullptr);
+	if (Playlist == nullptr)
+	{
+		return false;
+	}
+
 	int32 RemainingAttempts = Playlist->Num();
 
 	if (RemainingAttempts == 0)
@@ -417,9 +427,9 @@ bool UMediaPlayer::OpenFile(const FString& FilePath)
 
 	UE_LOG(LogMediaAssets, Verbose, TEXT("%s.OpenFile %s"), *GetFName().ToString(), *FilePath);
 
-	check(Playlist != nullptr);
+	EnsurePlaylist();
 
-	if (!Playlist->AddFile(FilePath))
+	if (Playlist == nullptr || !Playlist->AddFile(FilePath))
 	{
 		return false;
 	}
@@ -486,7 +496,14 @@ bool UMediaPlayer::OpenSourceInternal(UMediaSource* MediaSource, const FMediaPla
 		return false;
 	}
 
-	check(Playlist != nullptr);
+	EnsurePlaylist();
+
+	if (Playlist == nullptr)
+	{
+		UE_LOG(LogMediaAssets, Error, TEXT("Failed to create playlist on opening media source %s (%s)"), *MediaSource->GetName(), *MediaSource->GetUrl());
+		return false;
+	}
+
 	Playlist->Add(MediaSource);
 	PlayOnNext |= PlayerFacade->IsPlaying();
 	Playlist->GetNext(PlaylistIndex);
@@ -511,9 +528,9 @@ bool UMediaPlayer::OpenUrl(const FString& Url)
 
 	UE_LOG(LogMediaAssets, Verbose, TEXT("%s.OpenUrl %s"), *GetFName().ToString(), *Url);
 
-	check(Playlist != nullptr);
+	EnsurePlaylist();
 
-	if (!Playlist->AddUrl(Url))
+	if (Playlist == nullptr || !Playlist->AddUrl(Url))
 	{
 		return false;
 	}
@@ -552,7 +569,11 @@ bool UMediaPlayer::Previous()
 {
 	UE_LOG(LogMediaAssets, Verbose, TEXT("%s.Previous"), *GetFName().ToString());
 
-	check(Playlist != nullptr);
+	if (Playlist == nullptr)
+	{
+		return false;
+	}
+
 	int32 RemainingAttempts = Playlist->Num();
 
 	if (RemainingAttempts == 0)
@@ -582,6 +603,10 @@ bool UMediaPlayer::Previous()
 bool UMediaPlayer::Reopen()
 {
 	UE_LOG(LogMediaAssets, Verbose, TEXT("%s.Reopen"), *GetFName().ToString());
+	if (Playlist == nullptr)
+	{
+		return false;
+	}
 	return OpenPlaylistIndex(Playlist, PlaylistIndex);
 }
 
diff --git a/Engine/Source/Runtime/MediaAssets/Public/MediaPlayer.h b/Engine/Source/Runtime/MediaAssets/Public/MediaPlayer.h
--- a/Engine/Source/Runtime/MediaAssets/Public/MediaPlayer.h
+++ b/Engine/Source/Runtime/MediaAssets/Public/MediaPlayer.h
@@ -251,6 +251,7 @@ public:
 	UFUNCTION(BlueprintCallable, Category="Media|MediaPlayer")
 	UMediaPlaylist* GetPlaylist() const
 	{
+		EnsurePlaylist();
 		return Playlist;
 	}
 
@@ -973,7 +974,7 @@ public:
 	 */
 	UMediaPlaylist& GetPlaylistRef() const
 	{
-		check(Playlist != nullptr);
+		EnsurePlaylist();
 		return *Playlist;
 	}
 
@@ -1099,7 +1100,7 @@ protected:
 	 * @see OpenPlaylist, OpenPlaylistIndex
 	 */
 	UPROPERTY(BlueprintReadOnly, transient, Category=Playback)
-	UMediaPlaylist* Playlist;
+	mutable UMediaPlaylist* Playlist;
 
 	/**
 	 * The current index of the source in the play list being played.
@@ -1169,11 +1170,14 @@ private:
 	void HandlePlayerMediaEvent(EMediaEvent Event);
 
 	/** Sets the playlist and properly handles cases when this MediaPlayer object is in disregard for GC set */
-	void SetPlaylistInternal(UMediaPlaylist* InPlaylist);
+	void SetPlaylistInternal(UMediaPlaylist* InPlaylist) const;
 
 	/** Open media source with the given options. */
 	bool OpenSourceInternal(UMediaSource* MediaSource, const FMediaPlayerOptions* Options);
 
+	/** Ensure internal playlist is created */
+	void EnsurePlaylist() const;
+
 private:
 
 	/** An event delegate that is invoked when a media event occurred. */
