Adding explicit setup for playlist getter to UMediaPlayer
#rnx

diff --git a/Engine/Source/Runtime/MediaAssets/Public/MediaPlayer.h b/Engine/Source/Runtime/MediaAssets/Public/MediaPlayer.h
--- a/Engine/Source/Runtime/MediaAssets/Public/MediaPlayer.h
+++ b/Engine/Source/Runtime/MediaAssets/Public/MediaPlayer.h
@@ -1099,7 +1099,7 @@ protected:
 	 *
 	 * @see OpenPlaylist, OpenPlaylistIndex
 	 */
-	UPROPERTY(BlueprintReadOnly, transient, Category=Playback)
+	UPROPERTY(BlueprintReadOnly, transient, Category=Playback, BlueprintGetter = GetPlayList)
 	mutable UMediaPlaylist* Playlist;
 
 	/**
