Cleanup work in MediaTexture output code to avoid some useless conversions and enable mipped texture samples

diff --git a/Engine/Source/Editor/TextureEditor/Private/TextureEditorToolkit.cpp b/Engine/Source/Editor/TextureEditor/Private/TextureEditorToolkit.cpp
--- a/Engine/Source/Editor/TextureEditor/Private/TextureEditorToolkit.cpp
+++ b/Engine/Source/Editor/TextureEditor/Private/TextureEditorToolkit.cpp
@@ -44,6 +44,7 @@
 #include "Widgets/Input/SSlider.h"
 #include "Widgets/Layout/SSpacer.h"
 #include "Menus/TextureEditorViewOptionsMenu.h"
+#include "MediaTexture.h"
 
 #define LOCTEXT_NAMESPACE "FTextureEditorToolkit"
 
@@ -445,6 +446,7 @@ void FTextureEditorToolkit::PopulateQuickInfo( )
 	UTexture2DDynamic* Texture2DDynamic = Cast<UTexture2DDynamic>(Texture);
 	UVolumeTexture* VolumeTexture = Cast<UVolumeTexture>(Texture);
 	UTextureRenderTargetVolume* VolumeTextureRT = Cast<UTextureRenderTargetVolume>(Texture);
+	UMediaTexture* MediaTexture = Cast<UMediaTexture>(Texture);
 
 	const bool bIsVolume = VolumeTexture || VolumeTextureRT;
 	const bool bIsArray2D = Texture2DArray || Texture2DArrayRT;
@@ -613,6 +615,10 @@ void FTextureEditorToolkit::PopulateQuickInfo( )
 	{
 		NumMips = VolumeTextureRT->GetNumMips();
 	}
+	else if (MediaTexture)
+	{
+		NumMips = MediaTexture->GetTextureNumMips();
+	}
 
 	NumMipsText->SetText(FText::Format(NSLOCTEXT("TextureEditor", "QuickInfo_NumMips", "Number of Mips: {0}"), FText::AsNumber(NumMips)));
 
diff --git a/Engine/Source/Editor/TextureEditor/TextureEditor.Build.cs b/Engine/Source/Editor/TextureEditor/TextureEditor.Build.cs
--- a/Engine/Source/Editor/TextureEditor/TextureEditor.Build.cs
+++ b/Engine/Source/Editor/TextureEditor/TextureEditor.Build.cs
@@ -44,8 +44,10 @@ public class TextureEditor : ModuleRules
                 "EditorStyle",
 				"EditorFramework",
 				"UnrealEd",
-                "PropertyEditor"
-            }
+                "PropertyEditor",
+				"MediaUtils",
+				"MediaAssets"
+			}
 		);
 	}
 }
diff --git a/Engine/Source/Runtime/Media/Public/IMediaTextureSample.h b/Engine/Source/Runtime/Media/Public/IMediaTextureSample.h
--- a/Engine/Source/Runtime/Media/Public/IMediaTextureSample.h
+++ b/Engine/Source/Runtime/Media/Public/IMediaTextureSample.h
@@ -126,6 +126,17 @@ public:
 	 */
 	virtual FIntPoint GetDim() const = 0;
 
+	/**
+	 * Get the number of mips encoded in the sample
+	 *
+	 * @return Number of mips in the sample (including base level)
+	 * @note Default implementation provided as most samples will not feature mips
+	 */
+	virtual uint8 GetNumMips() const
+	{
+		return 1;
+	}
+
 	/**
 	 * Get the amount of time for which the sample is valid.
 	 *
diff --git a/Engine/Source/Runtime/Media/Public/IMediaTextureSampleConverter.h b/Engine/Source/Runtime/Media/Public/IMediaTextureSampleConverter.h
--- a/Engine/Source/Runtime/Media/Public/IMediaTextureSampleConverter.h
+++ b/Engine/Source/Runtime/Media/Public/IMediaTextureSampleConverter.h
@@ -22,7 +22,6 @@ public:
 
 	struct FConversionHints
 	{
-		bool bOutputSRGB;
 		uint8 NumMips;
 	};
 
diff --git a/Engine/Source/Runtime/MediaAssets/Private/Assets/MediaTexture.cpp b/Engine/Source/Runtime/MediaAssets/Private/Assets/MediaTexture.cpp
--- a/Engine/Source/Runtime/MediaAssets/Private/Assets/MediaTexture.cpp
+++ b/Engine/Source/Runtime/MediaAssets/Private/Assets/MediaTexture.cpp
@@ -71,6 +71,7 @@ UMediaTexture::UMediaTexture(const FObjectInitializer& ObjectInitializer)
 	, Dimensions(FIntPoint::ZeroValue)
 	, Size(0)
 	, CachedNextSampleTime(FTimespan::MinValue())
+	, TextureNumMips(1)
 {
 	NeverStream = true;
 	SRGB = true;
@@ -108,6 +109,10 @@ int32 UMediaTexture::GetWidth() const
 	return Dimensions.X;
 }
 
+int32 UMediaTexture::GetTextureNumMips() const
+{
+	return TextureNumMips;
+}
 
 void UMediaTexture::SetMediaPlayer(UMediaPlayer* NewMediaPlayer)
 {
@@ -148,7 +153,7 @@ FTextureResource* UMediaTexture::CreateResource()
 		}
 	}
 
-	Filter = (EnableGenMips && (NumMips > 1)) ? TF_Trilinear : TF_Bilinear;
+	Filter = (EnableGenMips && (TextureNumMips > 1)) ? TF_Trilinear : TF_Bilinear;
 
 	return new FMediaTextureResource(*this, Dimensions, Size, ClearColor, CurrentGuid.IsValid() ? CurrentGuid : DefaultGuid, EnableGenMips, NumMips);
 }
@@ -341,7 +346,7 @@ void UMediaTexture::TickResource(FTimespan Timecode)
 			if (CurrentPlayerPtr->GetPlayerFacade()->GetPlayer()->GetPlayerFeatureFlag(IMediaPlayer::EFeatureFlag::UsePlaybackTimingV2))
 			{
 				/*
-					We are using the old-style "1sample queue to sink" architecture to actually just pass long only ONE sample at a time from the logic
+					We are using the old-style "sample queue to sink" architecture to actually just pass along only ONE sample at a time from the logic
 					inside the player facade to the sinks. The selection as to what to render this frame is expected to be done earlier
 					this frame on the gamethread, hence only a single output frame is selected and passed along...
 				*/
@@ -363,14 +368,12 @@ void UMediaTexture::TickResource(FTimespan Timecode)
 				RenderParams.Rate = CurrentPlayerPtr->GetRate();
 				RenderParams.Time = Sample->GetTime();
 
-				if (NewStyleOutput)
-				{
-					// For new-style output the sample's sRGB state controls what we output
-					// (FOR NOW: this is too simplified if we have more then Rec703 material)
-					SRGB = Sample->IsOutputSrgb();
-					// Ensure sRGB changes will not trigger rendering the next time around
-					LastSrgb = SRGB;
-				}
+				// Track the sample's sRGB status and make sure the (high level) texture flags reflect it
+				// (this is really only good for Rec.703 type setups and will need enhancements!)
+				SRGB = Sample->IsOutputSrgb();
+				LastSrgb = SRGB;
+
+				TextureNumMips = (Sample->GetNumMips() > 1) ? Sample->GetNumMips() : NumMips;
 			}
 			else
 			{
@@ -381,6 +384,12 @@ void UMediaTexture::TickResource(FTimespan Timecode)
 				if (SampleQueue->Peek(Sample))
 				{
 					UpdateSampleInfo(Sample);
+
+					// See above: track sRGB state (for V1 we don't look at all samples - but this should be fine as this should not change on a per sample basis usually)
+					SRGB = Sample->IsOutputSrgb();
+					LastSrgb = SRGB;
+
+					TextureNumMips = (Sample->GetNumMips() > 1) ? Sample->GetNumMips() : NumMips;
 				}
 
 				RenderParams.SampleSource = SampleQueue;
@@ -406,14 +415,13 @@ void UMediaTexture::TickResource(FTimespan Timecode)
 	}
 
 	// update filter state, responding to mips setting
-	Filter = (EnableGenMips && (NumMips > 1)) ? TF_Trilinear : TF_Bilinear;
+	Filter = (EnableGenMips && (TextureNumMips > 1)) ? TF_Trilinear : TF_Bilinear;
 
 	// setup render parameters
 	RenderParams.CanClear = AutoClear;
 	RenderParams.ClearColor = ClearColor;
 	RenderParams.PreviousGuid = PreviousGuid;
 	RenderParams.CurrentGuid = CurrentGuid;
-	RenderParams.SrgbOutput = SRGB;
 	RenderParams.NumMips = NumMips;
 	
 	// redraw texture resource on render thread
diff --git a/Engine/Source/Runtime/MediaAssets/Private/Misc/MediaTextureResource.cpp b/Engine/Source/Runtime/MediaAssets/Private/Misc/MediaTextureResource.cpp
--- a/Engine/Source/Runtime/MediaAssets/Private/Misc/MediaTextureResource.cpp
+++ b/Engine/Source/Runtime/MediaAssets/Private/Misc/MediaTextureResource.cpp
@@ -93,38 +93,46 @@ namespace MediaTextureResourceHelpers
 	{
 		switch (Sample->GetFormat())
 		{
+			// 10-bit formats
 		case EMediaTextureSampleFormat::CharBGR10A2:
 		case EMediaTextureSampleFormat::YUVv210:
 			return PF_A2B10G10R10;
-		case EMediaTextureSampleFormat::Y416:
-			return PF_B8G8R8A8;
+			// Float formats
+		case EMediaTextureSampleFormat::FloatRGB:
+		case EMediaTextureSampleFormat::FloatRGBA:
+			return PF_FloatRGBA;
+			// Everything else maps to 8-bit RGB...
 		default:
 			return PF_B8G8R8A8;
 		}
 	}
 
-	/**
-	 * Check whether the given sample requires an sRGB texture.
-	 *
-	 * @param Sample The sample to check.
-	 * @return true if an sRGB texture is required, false otherwise.
-	 */
+
 	bool RequiresSrgbTexture(const TSharedPtr<IMediaTextureSample, ESPMode::ThreadSafe>& Sample)
+	{
+		return Sample->IsOutputSrgb();
+	}
+
+
+	bool RequiresSrgbInputTexture(const TSharedPtr<IMediaTextureSample, ESPMode::ThreadSafe>& Sample)
 	{
 		if (!Sample->IsOutputSrgb())
 		{
 			return false;
 		}
-
-		const EMediaTextureSampleFormat Format = Sample->GetFormat();
-
-		return ((Format == EMediaTextureSampleFormat::CharBGRA) ||
-				(Format == EMediaTextureSampleFormat::CharBMP) ||
-				(Format == EMediaTextureSampleFormat::FloatRGB) ||
-				(Format == EMediaTextureSampleFormat::FloatRGBA));
+		/*
+		* Input textures created to receive CPU side buffer sample data are ONLY created with sRGB attributes
+		* if we have any RGB(A) format. Any YUV (etc.) format does the conversion in SW when applicable in the
+		* conversion process!
+		*/
+		EMediaTextureSampleFormat Fmt = Sample->GetFormat();
+		return Fmt == EMediaTextureSampleFormat::CharBGRA ||
+			Fmt == EMediaTextureSampleFormat::CharBMP ||
+			Fmt == EMediaTextureSampleFormat::FloatRGB ||
+			Fmt == EMediaTextureSampleFormat::FloatRGBA;
 	}
-}
 
+} //namespace
 
 /* FMediaTextureResource structors
  *****************************************************************************/
@@ -179,36 +187,35 @@ void FMediaTextureResource::Render(const FRenderParams& Params)
 		bool UseSample;
 
 		// Yes, it it a queue?
-	if (SampleSource.IsValid())
-	{
-		// Yes, find out what we will display...
-		UseSample = false;
-
-		// get the most current sample to be rendered
-		TSharedPtr<IMediaTextureSample, ESPMode::ThreadSafe> TestSample;
-		while (SampleSource->Peek(TestSample) && TestSample.IsValid())
+		if (SampleSource.IsValid())
 		{
-				const FTimespan StartTime = TestSample->GetTime().Time;
-			const FTimespan EndTime = StartTime + TestSample->GetDuration();
+			// Yes, find out what we will display...
+			UseSample = false;
 
-			if ((Params.Rate >= 0.0f) && (Params.Time < StartTime))
+			// get the most current sample to be rendered
+			TSharedPtr<IMediaTextureSample, ESPMode::ThreadSafe> TestSample;
+			while (SampleSource->Peek(TestSample) && TestSample.IsValid())
 			{
-				break; // future sample (forward play)
-			}
+				const FTimespan StartTime = TestSample->GetTime().Time;
+				const FTimespan EndTime = StartTime + TestSample->GetDuration();
 
-			if ((Params.Rate <= 0.0f) && (Params.Time >= EndTime))
-			{
-				break; // future sample (reverse play)
-			}
+				if ((Params.Rate >= 0.0f) && (Params.Time < StartTime))
+				{
+					break; // future sample (forward play)
+				}
+
+				if ((Params.Rate <= 0.0f) && (Params.Time >= EndTime))
+				{
+					break; // future sample (reverse play)
+				}
 
 #if UE_MEDIAUTILS_DEVELOPMENT_DELEGATE
-			if (UseSample && Sample.IsValid())
-			{
-				FMediaDelegates::OnSampleDiscarded_RenderThread.Broadcast(&Owner, Sample);
-			}
+				if (UseSample && Sample.IsValid())
+				{
+					FMediaDelegates::OnSampleDiscarded_RenderThread.Broadcast(&Owner, Sample);
+				}
 #endif
-
-			UseSample = SampleSource->Dequeue(Sample);
+				UseSample = SampleSource->Dequeue(Sample);
 			}
 		}
 		else
@@ -243,7 +250,7 @@ void FMediaTextureResource::Render(const FRenderParams& Params)
 				//
 				// Sample dimensions are invalid
 				//
-				ClearTexture(FLinearColor::Red, Params.SrgbOutput); // mark corrupt sample
+				ClearTexture(FLinearColor::Red, false); // mark corrupt sample
 			}
 			else if (IMediaTextureSampleConverter *Converter = Sample->GetMediaTextureSampleConverter())
 			{
@@ -251,9 +258,10 @@ void FMediaTextureResource::Render(const FRenderParams& Params)
 				// Sample brings its own converter
 				//
 
+				const uint8 SampleNumMips = Sample->GetNumMips();
+
 				IMediaTextureSampleConverter::FConversionHints Hints;
-				Hints.bOutputSRGB = Params.SrgbOutput;
-				Hints.NumMips = NumMips;
+				Hints.NumMips = (SampleNumMips > 1) ? SampleNumMips : Params.NumMips;
 
 				// Does the conversion create its own output texture?
 				if ((Converter->GetConverterInfoFlags() & IMediaTextureSampleConverter::ConverterInfoFlags_WillCreateOutputTexture) == 0)
@@ -272,7 +280,7 @@ void FMediaTextureResource::Render(const FRenderParams& Params)
 					else
 					{
 						// Conversion is fully handled by converter
-						CreateOutputRenderTarget(Sample->GetOutputDim(), MediaTextureResourceHelpers::GetConvertedPixelFormat(Sample), Params.SrgbOutput, Params.ClearColor, NumMips);
+						CreateOutputRenderTarget(Sample->GetOutputDim(), MediaTextureResourceHelpers::GetConvertedPixelFormat(Sample), MediaTextureResourceHelpers::RequiresSrgbTexture(Sample), Params.ClearColor, Hints.NumMips);
 						Converter->Convert(RenderTargetTextureRHI, Hints);
 					}
 				}
@@ -283,10 +291,10 @@ void FMediaTextureResource::Render(const FRenderParams& Params)
 					if (Converter->Convert(OutTexture, Hints))
 					{
 						// As the converter created the texture, we might need to convert it even more to make it fit our needs. Check...
-						if (RequiresConversion(OutTexture, Sample->GetOutputDim(), Params.SrgbOutput, NumMips))
+						if (RequiresConversion(OutTexture, Sample->GetOutputDim(), NumMips))
 						{
-							CreateOutputRenderTarget(Sample->GetOutputDim(), MediaTextureResourceHelpers::GetConvertedPixelFormat(Sample), Params.SrgbOutput, Params.ClearColor, NumMips);
-							ConvertTextureToOutput(OutTexture.GetReference(), Sample, Params.SrgbOutput);
+							CreateOutputRenderTarget(Sample->GetOutputDim(), MediaTextureResourceHelpers::GetConvertedPixelFormat(Sample), MediaTextureResourceHelpers::RequiresSrgbTexture(Sample), Params.ClearColor, Hints.NumMips);
+							ConvertTextureToOutput(OutTexture.GetReference(), Sample);
 						}
 						else
 						{
@@ -299,24 +307,25 @@ void FMediaTextureResource::Render(const FRenderParams& Params)
 			}
 			else
 			{
+				// No custom conversion, we need default processing...
 				ConvertOrCopyNeeded = true;
 			}
 
 			if (ConvertOrCopyNeeded)
 			{
-				if (RequiresConversion(Sample, Params.SrgbOutput, NumMips))
+				if (RequiresConversion(Sample, NumMips))
 				{
 					//
 					// Sample needs to be converted by built in converter code
 					//
-					ConvertSample(Sample, Params.ClearColor, Params.SrgbOutput, NumMips);
+					ConvertSample(Sample, Params.ClearColor, NumMips);
 				}
 				else
 				{
 					//
 					// Sample can be used directly or is a simple copy
 					//
-					CopySample(Sample, Params.ClearColor, Params.SrgbOutput, NumMips, Params.CurrentGuid);
+					CopySample(Sample, Params.ClearColor, NumMips, Params.CurrentGuid);
 				}
 			}
 
@@ -343,7 +352,7 @@ void FMediaTextureResource::Render(const FRenderParams& Params)
 			}
 
 			// Generate mips as needed
-			if (CurrentNumMips > 1 && !Cleared)
+			if (CurrentNumMips > 1 && !Cleared && Sample->GetNumMips() == 1)
 			{
 				check(OutputTarget);
 
@@ -383,7 +392,7 @@ void FMediaTextureResource::Render(const FRenderParams& Params)
 		if (!Cleared || (Params.ClearColor != CurrentClearColor))
 		{
 			// Yes...
-			ClearTexture(Params.ClearColor, Params.SrgbOutput);
+			ClearTexture(Params.ClearColor, false);
 
 			if (CurrentSample)
 			{
@@ -538,20 +547,9 @@ void FMediaTextureResource::ReleaseDynamicRHI()
 }
 
 
- bool FMediaTextureResource::RequiresConversion(const FTexture2DRHIRef& SampleTexture, const FIntPoint & OutputDim, bool SrgbOutput, uint8 InNumMips) const
+ bool FMediaTextureResource::RequiresConversion(const FTexture2DRHIRef& SampleTexture, const FIntPoint & OutputDim, uint8 InNumMips) const
  {
-	 if (!Owner.NewStyleOutput)
-	 {
-		 //
-		 // Old Style
-		 //
-
-		 if (((SampleTexture->GetFlags() & TexCreate_SRGB) != 0) != SrgbOutput)
-		 {
-			 return true;
-		 }
-	 }
-	 else
+	 if (Owner.NewStyleOutput)
 	 {
 		 //
 		 // New Style
@@ -560,12 +558,9 @@ void FMediaTextureResource::ReleaseDynamicRHI()
 		 // For now we only allow this, single SRGB-style output format
 		 check(Owner.OutputFormat == MTOF_SRGB_LINOUT || Owner.OutputFormat == MTOF_Default);
 
-		 if ((SampleTexture->GetFlags() & TexCreate_SRGB) == 0)
-		 {
-			 return true;
-		 }
-
-		 if (SampleTexture->GetNumMips() != InNumMips)
+		 // If we have no mips in the sample, but want to have some in the output, we use the conversion pass
+		 // to setup level 0 and have a suitable output texture
+		 if (SampleTexture->GetNumMips() == 1 && InNumMips != 1)
 		 {
 			 return true;
 		 }
@@ -587,20 +582,9 @@ void FMediaTextureResource::ReleaseDynamicRHI()
  }
 
 
- bool FMediaTextureResource::RequiresConversion(const TSharedPtr<IMediaTextureSample, ESPMode::ThreadSafe>& Sample, bool SrgbOutput, uint8 InNumMips) const
+ bool FMediaTextureResource::RequiresConversion(const TSharedPtr<IMediaTextureSample, ESPMode::ThreadSafe>& Sample, uint8 InNumMips) const
  {
-	 if (!Owner.NewStyleOutput)
-	 {
-		 //
-		 // Old Style
-		 //
-
-		 if (Sample->IsOutputSrgb() != SrgbOutput)
-		 {
-			 return true;
-		 }
-	 }
-	 else
+	 if (Owner.NewStyleOutput)
 	 {
 		 //
 		 // New Style
@@ -609,14 +593,11 @@ void FMediaTextureResource::ReleaseDynamicRHI()
 		 // For now we only allow this, single SRGB-style output format
 		 check(Owner.OutputFormat == MTOF_SRGB_LINOUT || Owner.OutputFormat == MTOF_Default);
 
-		 if (!Sample->IsOutputSrgb())
-		 {
-			 return true;
-		 }
-
 		 FRHITexture *Texture = Sample->GetTexture();
 
-		 if (Texture && Texture->GetNumMips() != InNumMips)
+		 // If we have no mips in the sample, but want to have some in the output, we use the conversion pass
+		 // to setup level 0 and have a suitable output texture
+		 if (Texture && Texture->GetNumMips() == 1 && InNumMips != 1)
 		 {
 			 return true;
 		 }
@@ -641,10 +622,12 @@ void FMediaTextureResource::ReleaseDynamicRHI()
  }
 
 
-void FMediaTextureResource::ConvertSample(const TSharedPtr<IMediaTextureSample, ESPMode::ThreadSafe>& Sample, const FLinearColor& ClearColor, bool SrgbOutput, uint8 InNumMips)
+void FMediaTextureResource::ConvertSample(const TSharedPtr<IMediaTextureSample, ESPMode::ThreadSafe>& Sample, const FLinearColor& ClearColor, uint8 InNumMips)
 {
 	const EPixelFormat InputPixelFormat = MediaTextureResourceHelpers::GetPixelFormat(Sample);
 
+	const uint8 SampleNumMips = Sample->GetNumMips();
+
 	// get input texture
 	FRHITexture2D* InputTexture = nullptr;
 	{
@@ -667,19 +650,19 @@ void FMediaTextureResource::ConvertSample(const TSharedPtr<IMediaTextureSample,
 		{
 			// Make a source texture so we can convert from it...
 
-			const bool SrgbTexture = MediaTextureResourceHelpers::RequiresSrgbTexture(Sample);
-			 const ETextureCreateFlags InputCreateFlags = TexCreate_Dynamic | (SrgbTexture ? TexCreate_SRGB : TexCreate_None);  //<<< FILTER SRGB AWAY FOR G8 AND SUCH? (or will the RHI code kindly do this?)
+			const bool SrgbTexture = MediaTextureResourceHelpers::RequiresSrgbInputTexture(Sample);
+			const ETextureCreateFlags InputCreateFlags = TexCreate_Dynamic | (SrgbTexture ? TexCreate_SRGB : TexCreate_None);
 			const FIntPoint SampleDim = Sample->GetDim();
 
 			// create a new temp input render target if necessary
-			if (!InputTarget.IsValid() || (InputTarget->GetSizeXY() != SampleDim) || (InputTarget->GetFormat() != InputPixelFormat) || ((InputTarget->GetFlags() & InputCreateFlags) != InputCreateFlags))
+			if (!InputTarget.IsValid() || (InputTarget->GetSizeXY() != SampleDim) || (InputTarget->GetFormat() != InputPixelFormat) || ((InputTarget->GetFlags() & InputCreateFlags) != InputCreateFlags) || (InputTarget->GetNumMips() != SampleNumMips))
 			{
 				FRHIResourceCreateInfo CreateInfo;
 				InputTarget = RHICreateTexture2D(
 					SampleDim.X,
 					SampleDim.Y,
 					InputPixelFormat,
-					1,
+					SampleNumMips,
 					1,
 					InputCreateFlags,
 					CreateInfo);
@@ -688,8 +671,15 @@ void FMediaTextureResource::ConvertSample(const TSharedPtr<IMediaTextureSample,
 			}
 
 			// copy sample data to input render target
-			FUpdateTextureRegion2D Region(0, 0, 0, 0, SampleDim.X, SampleDim.Y);
-			RHIUpdateTexture2D(InputTarget, 0, Region, Sample->GetStride(), (uint8*)Sample->GetBuffer());
+			const uint8* Data = (const uint8*)Sample->GetBuffer();
+			for (uint8 MipLevel = 0; MipLevel < SampleNumMips; ++MipLevel)
+			{
+				uint32 Stride = Sample->GetStride() >> MipLevel;
+				uint32 Height = SampleDim.Y >> MipLevel;
+				FUpdateTextureRegion2D Region(0, 0, 0, 0, SampleDim.X >> MipLevel, Height);
+				RHIUpdateTexture2D(InputTarget, MipLevel, Region, Stride, Data);
+				Data += Stride * Height;
+			}
 
 			InputTexture = InputTarget;
 		}
@@ -697,17 +687,20 @@ void FMediaTextureResource::ConvertSample(const TSharedPtr<IMediaTextureSample,
 
 	// create the output texture
 	const FIntPoint OutputDim = Sample->GetOutputDim();
-	CreateOutputRenderTarget(OutputDim, MediaTextureResourceHelpers::GetConvertedPixelFormat(Sample), SrgbOutput, ClearColor, InNumMips);
+	CreateOutputRenderTarget(OutputDim, MediaTextureResourceHelpers::GetConvertedPixelFormat(Sample), MediaTextureResourceHelpers::RequiresSrgbTexture(Sample), ClearColor, (SampleNumMips > 1) ? SampleNumMips : InNumMips);
 
-	ConvertTextureToOutput(InputTexture, Sample, SrgbOutput);
+	ConvertTextureToOutput(InputTexture, Sample);
 }
 
 
- void FMediaTextureResource::ConvertTextureToOutput(FRHITexture2D* InputTexture, const TSharedPtr<IMediaTextureSample, ESPMode::ThreadSafe>& Sample, bool SrgbOutput)
+ void FMediaTextureResource::ConvertTextureToOutput(FRHITexture2D* InputTexture, const TSharedPtr<IMediaTextureSample, ESPMode::ThreadSafe>& Sample)
  {
 	// perform the conversion
 	FRHICommandListImmediate& CommandList = FRHICommandListExecutor::GetImmediateCommandList();
 	{
+		// We should never get here with a sample that contains mips!
+		check(Sample->GetNumMips() == 1);
+
 		SCOPED_DRAW_EVENT(CommandList, FMediaTextureResource_Convert);
 		SCOPED_GPU_STAT(CommandList, MediaTextureResource);
 
@@ -752,111 +745,105 @@ void FMediaTextureResource::ConvertSample(const TSharedPtr<IMediaTextureSample,
 			// Use the sample format to choose the conversion path
 			switch (Sample->GetFormat())
 			{
-			case EMediaTextureSampleFormat::CharAYUV:
-			{
-				TShaderMapRef<FAYUVConvertPS> ConvertShader(ShaderMap);
-				GraphicsPSOInit.BoundShaderState.PixelShaderRHI = ConvertShader.GetPixelShader();
-				SetGraphicsPipelineState(CommandList, GraphicsPSOInit);
-				ConvertShader->SetParameters(CommandList, InputTexture, YUVToRGBMatrix, YUVOffset, bIsSampleOutputSrgb);
-			}
-			break;
-
-			case EMediaTextureSampleFormat::CharBMP:
-			{
-				TShaderMapRef<FBMPConvertPS> ConvertShader(ShaderMap);
-				GraphicsPSOInit.BoundShaderState.PixelShaderRHI = ConvertShader.GetPixelShader();
-				SetGraphicsPipelineState(CommandList, GraphicsPSOInit);
-				ConvertShader->SetParameters(CommandList, InputTexture, OutputDim, bIsSampleOutputSrgb && !SrgbOutput);
-			}
-			break;
-
-			case EMediaTextureSampleFormat::CharNV12:
-			{
-				if (InputTexture->GetFormat() == PF_NV12)
+				case EMediaTextureSampleFormat::CharAYUV:
 				{
-					TShaderMapRef<FNV12ConvertPS> ConvertShader(ShaderMap);
+					TShaderMapRef<FAYUVConvertPS> ConvertShader(ShaderMap);
 					GraphicsPSOInit.BoundShaderState.PixelShaderRHI = ConvertShader.GetPixelShader();
 					SetGraphicsPipelineState(CommandList, GraphicsPSOInit);
-					ConvertShader->SetParameters(CommandList, InputTexture, OutputDim, YUVToRGBMatrix, YUVOffset, bIsSampleOutputSrgb);
+					ConvertShader->SetParameters(CommandList, InputTexture, YUVToRGBMatrix, YUVOffset, bIsSampleOutputSrgb);
 				}
-			else
+				break;
+
+				case EMediaTextureSampleFormat::CharBMP:
 				{
-					TShaderMapRef<FNV12ConvertAsBytesPS> ConvertShader(ShaderMap);
+					// Simple 1:1 copy plus flip & color adjustment (but using normal texture sampler: sRGB conversions may occur depending on setup; any manual sRGB/linear conversion is disabled)
+					TShaderMapRef<FBMPConvertPS> ConvertShader(ShaderMap);
 					GraphicsPSOInit.BoundShaderState.PixelShaderRHI = ConvertShader.GetPixelShader();
 					SetGraphicsPipelineState(CommandList, GraphicsPSOInit);
-					ConvertShader->SetParameters(CommandList, InputTexture, OutputDim, YUVToRGBMatrix, YUVOffset, bIsSampleOutputSrgb);
+					ConvertShader->SetParameters(CommandList, InputTexture, OutputDim, false);
 				}
-			}
-			break;
+				break;
 
-			case EMediaTextureSampleFormat::CharNV21:
-			{
-				// source texture might be NV12 or G8...
-				TShaderMapRef<FNV21ConvertPS> ConvertShader(ShaderMap);
-				GraphicsPSOInit.BoundShaderState.PixelShaderRHI = ConvertShader.GetPixelShader();
-				SetGraphicsPipelineState(CommandList, GraphicsPSOInit);
-				ConvertShader->SetParameters(CommandList, InputTexture, OutputDim, YUVToRGBMatrix, YUVOffset, bIsSampleOutputSrgb);
-			}
-			break;
+				case EMediaTextureSampleFormat::CharNV12:
+				{
+					if (InputTexture->GetFormat() == PF_NV12)
+					{
+						TShaderMapRef<FNV12ConvertPS> ConvertShader(ShaderMap);
+						GraphicsPSOInit.BoundShaderState.PixelShaderRHI = ConvertShader.GetPixelShader();
+						SetGraphicsPipelineState(CommandList, GraphicsPSOInit);
+						ConvertShader->SetParameters(CommandList, InputTexture, OutputDim, YUVToRGBMatrix, YUVOffset, bIsSampleOutputSrgb);
+					}
+					else
+					{
+						TShaderMapRef<FNV12ConvertAsBytesPS> ConvertShader(ShaderMap);
+						GraphicsPSOInit.BoundShaderState.PixelShaderRHI = ConvertShader.GetPixelShader();
+						SetGraphicsPipelineState(CommandList, GraphicsPSOInit);
+						ConvertShader->SetParameters(CommandList, InputTexture, OutputDim, YUVToRGBMatrix, YUVOffset, bIsSampleOutputSrgb);
+					}
+				}
+				break;
 
-			case EMediaTextureSampleFormat::CharUYVY:
-			{
-				TShaderMapRef<FUYVYConvertPS> ConvertShader(ShaderMap);
-				GraphicsPSOInit.BoundShaderState.PixelShaderRHI = ConvertShader.GetPixelShader();
-				SetGraphicsPipelineState(CommandList, GraphicsPSOInit);
-				ConvertShader->SetParameters(CommandList, InputTexture, YUVToRGBMatrix, YUVOffset, bIsSampleOutputSrgb);
-			}
-			break;
+				case EMediaTextureSampleFormat::CharNV21:
+				{
+					// source texture might be NV12 or G8...
+					TShaderMapRef<FNV21ConvertPS> ConvertShader(ShaderMap);
+					GraphicsPSOInit.BoundShaderState.PixelShaderRHI = ConvertShader.GetPixelShader();
+					SetGraphicsPipelineState(CommandList, GraphicsPSOInit);
+					ConvertShader->SetParameters(CommandList, InputTexture, OutputDim, YUVToRGBMatrix, YUVOffset, bIsSampleOutputSrgb);
+				}
+				break;
 
-			case EMediaTextureSampleFormat::CharYUY2:
-			{
-				TShaderMapRef<FYUY2ConvertPS> ConvertShader(ShaderMap);
-				GraphicsPSOInit.BoundShaderState.PixelShaderRHI = ConvertShader.GetPixelShader();
-				SetGraphicsPipelineState(CommandList, GraphicsPSOInit);
-				ConvertShader->SetParameters(CommandList, InputTexture, OutputDim, YUVToRGBMatrix, YUVOffset, bIsSampleOutputSrgb);
-			}
-			break;
+				case EMediaTextureSampleFormat::CharUYVY:
+				{
+					TShaderMapRef<FUYVYConvertPS> ConvertShader(ShaderMap);
+					GraphicsPSOInit.BoundShaderState.PixelShaderRHI = ConvertShader.GetPixelShader();
+					SetGraphicsPipelineState(CommandList, GraphicsPSOInit);
+					ConvertShader->SetParameters(CommandList, InputTexture, YUVToRGBMatrix, YUVOffset, bIsSampleOutputSrgb);
+				}
+				break;
 
-			case EMediaTextureSampleFormat::CharYVYU:
-			{
-				TShaderMapRef<FYVYUConvertPS> ConvertShader(ShaderMap);
-				GraphicsPSOInit.BoundShaderState.PixelShaderRHI = ConvertShader.GetPixelShader();
-				SetGraphicsPipelineState(CommandList, GraphicsPSOInit);
-				ConvertShader->SetParameters(CommandList, InputTexture, YUVToRGBMatrix, YUVOffset, bIsSampleOutputSrgb);
-			}
-			break;
+				case EMediaTextureSampleFormat::CharYUY2:
+				{
+					TShaderMapRef<FYUY2ConvertPS> ConvertShader(ShaderMap);
+					GraphicsPSOInit.BoundShaderState.PixelShaderRHI = ConvertShader.GetPixelShader();
+					SetGraphicsPipelineState(CommandList, GraphicsPSOInit);
+					ConvertShader->SetParameters(CommandList, InputTexture, OutputDim, YUVToRGBMatrix, YUVOffset, bIsSampleOutputSrgb);
+				}
+				break;
 
-			case EMediaTextureSampleFormat::YUVv210:
-			{
-				TShaderMapRef<FYUVv210ConvertPS> ConvertShader(ShaderMap);
-				GraphicsPSOInit.BoundShaderState.PixelShaderRHI = ConvertShader.GetPixelShader();
-				SetGraphicsPipelineState(CommandList, GraphicsPSOInit);
-				ConvertShader->SetParameters(CommandList, InputTexture, OutputDim, YUVToRGBMatrix, YUVOffset, bIsSampleOutputSrgb);
-			}
-			break;
+				case EMediaTextureSampleFormat::CharYVYU:
+				{
+					TShaderMapRef<FYVYUConvertPS> ConvertShader(ShaderMap);
+					GraphicsPSOInit.BoundShaderState.PixelShaderRHI = ConvertShader.GetPixelShader();
+					SetGraphicsPipelineState(CommandList, GraphicsPSOInit);
+					ConvertShader->SetParameters(CommandList, InputTexture, YUVToRGBMatrix, YUVOffset, bIsSampleOutputSrgb);
+				}
+				break;
 
-			case EMediaTextureSampleFormat::CharBGR10A2:
-			{
-				TShaderMapRef<FRGBConvertPS> ConvertShader(ShaderMap);
-				GraphicsPSOInit.BoundShaderState.PixelShaderRHI = ConvertShader.GetPixelShader();
-				SetGraphicsPipelineState(CommandList, GraphicsPSOInit);
-				ConvertShader->SetParameters(CommandList, InputTexture, OutputDim, bIsSampleOutputSrgb);
-			}
-			break;
+				case EMediaTextureSampleFormat::YUVv210:
+				{
+					TShaderMapRef<FYUVv210ConvertPS> ConvertShader(ShaderMap);
+					GraphicsPSOInit.BoundShaderState.PixelShaderRHI = ConvertShader.GetPixelShader();
+					SetGraphicsPipelineState(CommandList, GraphicsPSOInit);
+					ConvertShader->SetParameters(CommandList, InputTexture, OutputDim, YUVToRGBMatrix, YUVOffset, bIsSampleOutputSrgb);
+				}
+				break;
 
-			case EMediaTextureSampleFormat::CharBGRA:
-			case EMediaTextureSampleFormat::FloatRGB:
-			case EMediaTextureSampleFormat::FloatRGBA:
-			{
-				// Simple 1:1 copy (but using normal texture sampler: sRGB conversions may occur depending on setup; any manuao sRGB/linear conversion is disabled)
-				TShaderMapRef<FRGBConvertPS> ConvertShader(ShaderMap);
-				GraphicsPSOInit.BoundShaderState.PixelShaderRHI = ConvertShader.GetPixelShader();
-				SetGraphicsPipelineState(CommandList, GraphicsPSOInit);
-				ConvertShader->SetParameters(CommandList, InputTexture, OutputDim, false);
-			}
-			break;
+				case EMediaTextureSampleFormat::CharBGR10A2:
+				case EMediaTextureSampleFormat::CharBGRA:
+				case EMediaTextureSampleFormat::FloatRGB:
+				case EMediaTextureSampleFormat::FloatRGBA:
+				{
+					// Simple 1:1 copy (we have a real sRGB texture here if sRGB  is encoded)
+					// (as RGB formats normally do not see any conversion, we only get here if mips need to be generated - in that case this populates mip level 0)
+					TShaderMapRef<FRGBConvertPS> ConvertShader(ShaderMap);
+					GraphicsPSOInit.BoundShaderState.PixelShaderRHI = ConvertShader.GetPixelShader();
+					SetGraphicsPipelineState(CommandList, GraphicsPSOInit);
+					ConvertShader->SetParameters(CommandList, InputTexture, OutputDim, false);
+				}
+				break;
 
-			default:
+				default:
 				{
 					// This should not happen in normal use: still - end the render pass to avoid any trouble with RHI
 					CommandList.EndRenderPass();
@@ -880,10 +867,11 @@ void FMediaTextureResource::ConvertSample(const TSharedPtr<IMediaTextureSample,
 }
 
 
-void FMediaTextureResource::CopySample(const TSharedPtr<IMediaTextureSample, ESPMode::ThreadSafe>& Sample, const FLinearColor& ClearColor, bool SrgbOutput, uint8 InNumMips, const FGuid & TextureGUID)
+void FMediaTextureResource::CopySample(const TSharedPtr<IMediaTextureSample, ESPMode::ThreadSafe>& Sample, const FLinearColor& ClearColor, uint8 InNumMips, const FGuid & TextureGUID)
 {
 	FRHITexture* SampleTexture = Sample->GetTexture();
 	FRHITexture2D* SampleTexture2D = (SampleTexture != nullptr) ? SampleTexture->GetTexture2D() : nullptr;
+	const uint8 SampleNumMips = Sample->GetNumMips();
 
 	// If the sample already provides a texture resource, we simply use that
 	// as the output render target. If the sample only provides raw data, then
@@ -902,7 +890,7 @@ void FMediaTextureResource::CopySample(const TSharedPtr<IMediaTextureSample, ESP
 		else
 		{
 			// Texture to receive texture from sample
-			CreateOutputRenderTarget(Sample->GetOutputDim(), MediaTextureResourceHelpers::GetPixelFormat(Sample), SrgbOutput, ClearColor, InNumMips);
+			CreateOutputRenderTarget(Sample->GetOutputDim(), MediaTextureResourceHelpers::GetPixelFormat(Sample), MediaTextureResourceHelpers::RequiresSrgbTexture(Sample), ClearColor, (SampleNumMips > 1) ? SampleNumMips : InNumMips);
 
 			// Copy data into the output texture to able to add mips later on
 			FRHICommandListExecutor::GetImmediateCommandList().CopyTexture(SampleTexture2D, OutputTarget, FRHICopyTextureInfo());
@@ -911,12 +899,15 @@ void FMediaTextureResource::CopySample(const TSharedPtr<IMediaTextureSample, ESP
 	else
 	{
 		// Texture to receive precisely only output pixels via CPU copy
-		CreateOutputRenderTarget(Sample->GetDim(), MediaTextureResourceHelpers::GetPixelFormat(Sample), SrgbOutput, ClearColor, InNumMips);
+		CreateOutputRenderTarget(Sample->GetDim(), MediaTextureResourceHelpers::GetPixelFormat(Sample), MediaTextureResourceHelpers::RequiresSrgbTexture(Sample), ClearColor, (SampleNumMips > 1) ? SampleNumMips : InNumMips);
 
 		// If we also have no source buffer and the platform generally would allow for use of external textures, we assume it is just that...
 		// (as long as the player actually produces (dummy) samples, this will enable mips support as well as auto conversion for "new style output" mode)
 		if (!Sample->GetBuffer())
 		{
+			// we expect an external texture image in this case - we should have no mips reported here!
+			check(SampleNumMips == 1);
+
 			if (GSupportsImageExternal)
 			{
 				CopyFromExternalTexture(Sample, TextureGUID);
@@ -937,8 +928,16 @@ void FMediaTextureResource::CopySample(const TSharedPtr<IMediaTextureSample, ESP
 		else
 		{
 			// Copy sample data (from CPU mem) to output render target
-			FUpdateTextureRegion2D Region(0, 0, 0, 0, Sample->GetDim().X, Sample->GetDim().Y);
-			RHIUpdateTexture2D(RenderTargetTextureRHI.GetReference(), 0, Region, Sample->GetStride(), (uint8*)Sample->GetBuffer());
+			const FIntPoint SampleDim = Sample->GetDim();
+			const uint8* Data = (const uint8*)Sample->GetBuffer();
+			for (uint8 MipLevel = 0; MipLevel < SampleNumMips; ++MipLevel)
+			{
+				uint32 Stride = Sample->GetStride() >> MipLevel;
+				uint32 Height = SampleDim.Y >> MipLevel;
+				FUpdateTextureRegion2D Region(0, 0, 0, 0, SampleDim.X >> MipLevel, Height);
+				RHIUpdateTexture2D(RenderTargetTextureRHI, MipLevel, Region, Stride, Data);
+				Data += Stride * Height;
+			}
 
 			// Make sure resource is in SRV mode again
 			FRHICommandListExecutor::GetImmediateCommandList().Transition(FRHITransitionInfo(RenderTargetTextureRHI.GetReference(), ERHIAccess::Unknown, ERHIAccess::SRVMask));
@@ -1038,6 +1037,7 @@ void FMediaTextureResource::UpdateTextureReference(FRHITexture2D* NewTexture)
 	RenderTargetTextureRHI = NewTexture;
 
 	RHIUpdateTextureReference(Owner.TextureReference.TextureReferenceRHI, NewTexture);
+	// note: sRGB status for Owner.SRGB is handled (on game thread) in MediaTetxure.cpp
 
 	if (RenderTargetTextureRHI != nullptr)
 	{
diff --git a/Engine/Source/Runtime/MediaAssets/Private/Misc/MediaTextureResource.h b/Engine/Source/Runtime/MediaAssets/Private/Misc/MediaTextureResource.h
--- a/Engine/Source/Runtime/MediaAssets/Private/Misc/MediaTextureResource.h
+++ b/Engine/Source/Runtime/MediaAssets/Private/Misc/MediaTextureResource.h
@@ -87,9 +87,6 @@ public:
 		/** The player facade that provides the video samples to render. */
 		TWeakPtr<FMediaTextureSampleSource, ESPMode::ThreadSafe> SampleSource;
 
-		/** Whether output should be in sRGB color space. */
-		bool SrgbOutput;
-
 		/** Number of mips wanted */
 		uint8 NumMips;
 
@@ -147,13 +144,12 @@ protected:
 	 *
 	 * @param Sample The texture sample to convert.
 	 * @param ClearColor The clear color to use for the output texture.
-	 * @param SrgbOutput Whether the output texture is in sRGB color space.
 	 * @param Number of mips
 	 * @see CopySample
 	 */
-	void ConvertSample(const TSharedPtr<IMediaTextureSample, ESPMode::ThreadSafe>& Sample, const FLinearColor& ClearColor, bool SrgbOutput, uint8 InNumMips);
+	void ConvertSample(const TSharedPtr<IMediaTextureSample, ESPMode::ThreadSafe>& Sample, const FLinearColor& ClearColor, uint8 InNumMips);
 
-	void ConvertTextureToOutput(FRHITexture2D* InputTexture, const TSharedPtr<IMediaTextureSample, ESPMode::ThreadSafe>& Sample, bool SrgbOutput);
+	void ConvertTextureToOutput(FRHITexture2D* InputTexture, const TSharedPtr<IMediaTextureSample, ESPMode::ThreadSafe>& Sample);
 
 	/**
 	 * Render the given texture sample by using it as or copying it to the render target.
@@ -164,7 +160,7 @@ protected:
 	 * @param Number of mips
 	 * @see ConvertSample
 	 */
-	void CopySample(const TSharedPtr<IMediaTextureSample, ESPMode::ThreadSafe>& Sample, const FLinearColor& ClearColor, bool SrgbOutput, uint8 InNumMips, const FGuid & TextureGUID);
+	void CopySample(const TSharedPtr<IMediaTextureSample, ESPMode::ThreadSafe>& Sample, const FLinearColor& ClearColor, uint8 InNumMips, const FGuid & TextureGUID);
 
 	/** Calculates the current resource size and notifies the owner texture. */
 	void UpdateResourceSize();
@@ -194,8 +190,8 @@ protected:
 	/** Copy to local buffer from external texture */
 	void CopyFromExternalTexture(const TSharedPtr <IMediaTextureSample, ESPMode::ThreadSafe>& Sample, const FGuid & TextureGUID);
 
-	bool RequiresConversion(const TSharedPtr<IMediaTextureSample, ESPMode::ThreadSafe>& Sample, bool SrgbOutput, uint8 numMips) const;
-	bool RequiresConversion(const FTexture2DRHIRef& SampleTexture, const FIntPoint & OutputDim, bool SrgbOutput, uint8 numMips) const;
+	bool RequiresConversion(const TSharedPtr<IMediaTextureSample, ESPMode::ThreadSafe>& Sample, uint8 numMips) const;
+	bool RequiresConversion(const FTexture2DRHIRef& SampleTexture, const FIntPoint & OutputDim, uint8 numMips) const;
 
 private:
 
