Submitting missing file to fix build issues
#rnx

diff --git a/Engine/Source/Runtime/MediaAssets/Public/MediaTexture.h b/Engine/Source/Runtime/MediaAssets/Public/MediaTexture.h
--- a/Engine/Source/Runtime/MediaAssets/Public/MediaTexture.h
+++ b/Engine/Source/Runtime/MediaAssets/Public/MediaTexture.h
@@ -107,15 +107,6 @@ public:
 	UFUNCTION(BlueprintCallable, Category="Media|MediaTexture")
 	int32 GetHeight() const;
 
-	/**
-	 * Get the media player that provides the video samples.
-	 *
-	 * @return The texture's media player, or nullptr if not set.
-	 * @see SetMediaPlayer
-	 */
-	UFUNCTION(BlueprintCallable, Category="Media|MediaTexture")
-	UMediaPlayer* GetMediaPlayer() const;
-
 	/**
 	 * Gets the current width of the texture.
 	 *
@@ -125,6 +116,23 @@ public:
 	UFUNCTION(BlueprintCallable, Category="Media|MediaTexture")
 	int32 GetWidth() const;
 
+	/**
+	 * Gets the current numbe of mips of the texture.
+	 *
+	 * @return Number of mips.
+	 */
+	UFUNCTION(BlueprintCallable, Category = "Media|MediaTexture")
+	int32 GetTextureNumMips() const;
+
+	/**
+	 * Get the media player that provides the video samples.
+	 *
+	 * @return The texture's media player, or nullptr if not set.
+	 * @see SetMediaPlayer
+	 */
+	UFUNCTION(BlueprintCallable, Category = "Media|MediaTexture")
+	UMediaPlayer* GetMediaPlayer() const;
+
 	/**
 	 * Set the media player that provides the video samples.
 	 *
@@ -272,4 +280,7 @@ private:
 
 	/** Next available sample time when last render call was made */
 	TAtomic<FTimespan> CachedNextSampleTime;
+
+	/** Number of mips in the actual output texture */
+	int32 TextureNumMips;
 };
