Changed the retirement logic for logically "done" texture samples to not need a timeout etc.
#rnx

diff --git a/Engine/Source/Runtime/MediaAssets/Private/Misc/MediaTextureResource.cpp b/Engine/Source/Runtime/MediaAssets/Private/Misc/MediaTextureResource.cpp
--- a/Engine/Source/Runtime/MediaAssets/Private/Misc/MediaTextureResource.cpp
+++ b/Engine/Source/Runtime/MediaAssets/Private/Misc/MediaTextureResource.cpp
@@ -6,6 +6,8 @@
 #include "DeviceProfiles/DeviceProfile.h"
 #include "DeviceProfiles/DeviceProfileManager.h"
 #include "ExternalTexture.h"
+#include "IMediaModule.h"
+#include "IMediaClock.h"
 #include "IMediaPlayer.h"
 #include "IMediaSamples.h"
 #include "IMediaTextureSample.h"
@@ -41,6 +43,114 @@ CSV_DECLARE_CATEGORY_MODULE_EXTERN(MEDIA_API, MediaStreaming);
 
 DECLARE_GPU_STAT_NAMED(MediaTextureResource, TEXT("MediaTextureResource"));
 
+
+/* GPU data deletion helper
+ *****************************************************************************/
+
+template<typename ObjectRefType> struct TGPUsyncedDataDeleter : public IMediaClockSink
+{
+	static TSharedRef<TGPUsyncedDataDeleter<ObjectRefType>, ESPMode::ThreadSafe> Create()
+	{
+		auto Ret = MakeShared<TGPUsyncedDataDeleter<ObjectRefType>, ESPMode::ThreadSafe>();
+		Ret->WeakThis = Ret;
+		return Ret;
+	}
+
+	virtual ~TGPUsyncedDataDeleter()
+	{
+		// See if all samples are ready to be retired now...
+		if (!Update())
+		{
+			// They are. No need for any async task...
+			return;
+		}
+
+		// Start having us ticked so we can retire the rest over the next game loop iterations...
+		IMediaModule* MediaModule = FModuleManager::GetModulePtr<IMediaModule>(TEXT("MediaModule"));
+		if (MediaModule)
+		{
+			MediaModule->GetClock().AddSink(StaticCastSharedRef<IMediaClockSink, TGPUsyncedDataDeleter<ObjectRefType>, ESPMode::ThreadSafe>(WeakThis.Pin().ToSharedRef()));
+		}
+	}
+
+	void Retire(const ObjectRefType& Object)
+	{
+		FRHICommandListImmediate& CommandList = FRHICommandListExecutor::GetImmediateCommandList();
+
+		// Prep "retirement package"
+		FRetiringObjectInfo Info;
+		Info.Object = Object;
+		Info.GPUFence = CommandList.CreateGPUFence(TEXT("MediaTextureResourceReuseFence"));
+		Info.RetireTime = FPlatformTime::Seconds();
+
+		// Insert fence. We assume that GPU-workload-wise this marks the spot usage of the sample is done
+		CommandList.WriteGPUFence(Info.GPUFence);
+
+		// Recall for later checking...
+		FScopeLock Lock(&CS);
+		Objects.Push(Info);
+	}
+
+	bool Update()
+	{
+		FScopeLock Lock(&CS);
+
+		// Check for any retired samples that are not done being touched by the GPU...
+		int32 Idx = 0;
+		for (; Idx < Objects.Num(); ++Idx)
+		{
+			// Either no fence present or the fence has been signaled?
+			if (Objects[Idx].GPUFence.IsValid() && !Objects[Idx].GPUFence->Poll())
+			{
+				// No. This one is still busy, we can stop...
+				break;
+			}
+		}
+		// Remove (hence return to the pool / free up fence) all the finished ones...
+		if (Idx != 0)
+		{
+			Objects.RemoveAt(0, Idx);
+		}
+		return Objects.Num() != 0;
+	}
+
+private:
+	virtual void TickInput(FTimespan DeltaTime, FTimespan Timecode) override
+	{
+		// Check what data is still not signaled...
+		if (!Update())
+		{
+			// All is gone! Remove reference to this (which will delete this instance)
+			IMediaModule* MediaModule = FModuleManager::GetModulePtr<IMediaModule>(TEXT("MediaModule"));
+			if (MediaModule)
+			{
+				MediaModule->GetClock().RemoveSink(StaticCastSharedRef<IMediaClockSink, TGPUsyncedDataDeleter<ObjectRefType>, ESPMode::ThreadSafe>(WeakThis.Pin().ToSharedRef()));
+			}
+		}
+	}
+
+	struct FRetiringObjectInfo
+	{
+		ObjectRefType Object;
+		FGPUFenceRHIRef GPUFence;
+		double RetireTime;
+	};
+
+	TWeakPtr<TGPUsyncedDataDeleter<ObjectRefType>, ESPMode::ThreadSafe> WeakThis;
+	TArray<FRetiringObjectInfo> Objects;
+	FCriticalSection CS;
+};
+
+struct FPriorSamples : public TGPUsyncedDataDeleter<TSharedPtr<IMediaTextureSample, ESPMode::ThreadSafe>>
+{
+	using BaseType = TGPUsyncedDataDeleter<TSharedPtr<IMediaTextureSample, ESPMode::ThreadSafe>>;
+
+	static TSharedRef<FPriorSamples, ESPMode::ThreadSafe> Create()
+	{
+		return StaticCastSharedRef<FPriorSamples, BaseType, ESPMode::ThreadSafe>(BaseType::Create());
+	}
+};
+
 /* Local helpers
  *****************************************************************************/
 
@@ -147,7 +257,7 @@ FMediaTextureResource::FMediaTextureResource(UMediaTexture& InOwner, FIntPoint&
 	, bEnableGenMips(InEnableGenMips)
 	, CurrentNumMips(InEnableGenMips ? InNumMips : 1)
 	, CurrentSamplerFilter(ESamplerFilter_Num)
-	, PriorSamples(MakeShared<FPriorSamples, ESPMode::ThreadSafe>())
+	, PriorSamples(FPriorSamples::Create())
 {
 #if PLATFORM_ANDROID
 	bUsesImageExternal = !Owner.NewStyleOutput && (!FAndroidMisc::ShouldUseVulkan() && GSupportsImageExternal);
@@ -159,7 +269,7 @@ FMediaTextureResource::FMediaTextureResource(UMediaTexture& InOwner, FIntPoint&
 
 void FMediaTextureResource::FlushPendingData()
 {
-	PriorSamples = MakeShared<FPriorSamples, ESPMode::ThreadSafe>();
+	PriorSamples = FPriorSamples::Create();
 }
 
 /* FMediaTextureResource interface
diff --git a/Engine/Source/Runtime/MediaAssets/Private/Misc/MediaTextureResource.h b/Engine/Source/Runtime/MediaAssets/Private/Misc/MediaTextureResource.h
--- a/Engine/Source/Runtime/MediaAssets/Private/Misc/MediaTextureResource.h
+++ b/Engine/Source/Runtime/MediaAssets/Private/Misc/MediaTextureResource.h
@@ -23,20 +23,11 @@ class IMediaPlayer;
 class IMediaTextureSample;
 class UMediaTexture;
 struct FGenerateMipsStruct;
+struct FPriorSamples;
 
 enum class EMediaTextureSinkFormat;
 enum class EMediaTextureSinkMode;
 
-#if PLATFORM_WINDOWS || (defined(PLATFORM_PS4) && PLATFORM_PS4) || (defined(PLATFORM_PS5) && PLATFORM_PS5)
-#define USE_LIMITED_FENCEWAIT	1
-#else
-#define USE_LIMITED_FENCEWAIT	0
-#endif
-
-#if USE_LIMITED_FENCEWAIT
-static const double MaxWaitForFence = 2.0;	// HACK: wait a max of 2s for a GPU fence, then assume we will never see it signal & pretent it did signal
-#endif
-
 /**
  * Texture resource type for media textures.
  */
@@ -238,125 +229,6 @@ private:
 	TSharedPtr<IMediaTextureSample, ESPMode::ThreadSafe> CurrentSample;
 
 	/** prior samples not yet ready for retirement as GPU may still actively use them */
-	template<typename ObjectRefType> struct TGPUsyncedDataDeleter
-	{
-		~TGPUsyncedDataDeleter()
-		{
-			Flush();
-		}
-
-		void Retire(const ObjectRefType& Object)
-		{
-			FRHICommandListImmediate& CommandList = FRHICommandListExecutor::GetImmediateCommandList();
-
-			// Prep "retirement package"
-			FRetiringObjectInfo Info;
-			Info.Object = Object;
-			Info.GPUFence = CommandList.CreateGPUFence(TEXT("MediaTextureResourceReuseFence"));
-			Info.RetireTime = FPlatformTime::Seconds();
-
-			// Insert fence. We assume that GPU-workload-wise this marks the spot usage of the sample is done
-			CommandList.WriteGPUFence(Info.GPUFence);
-
-			// Recall for later checking...
-			FScopeLock Lock(&CS);
-			Objects.Push(Info);
-		}
-
-		bool Update()
-		{
-			FScopeLock Lock(&CS);
-
-			// Check for any retired samples that are not done being touched by the GPU...
-			int32 Idx = 0;
-			for (; Idx < Objects.Num(); ++Idx)
-			{
-#if USE_LIMITED_FENCEWAIT
-				double Now = FPlatformTime::Seconds();
-#endif
-
-				// Either no fence present or the fence has been signaled?
-				if (Objects[Idx].GPUFence.IsValid() && !Objects[Idx].GPUFence->Poll())
-				{
-					// No. This one is still busy, we can stop...
-
-#if USE_LIMITED_FENCEWAIT
-					// HACK: But how long has this been going on? Might we have a fence that never will signal?
-					if ((Now - Objects[Idx].RetireTime) < MaxWaitForFence)
-#else
-					if (1)
-#endif
-					{
-						break;
-					}
-				}
-			}
-			// Remove (hence return to the pool / free up fence) all the finished ones...
-			if (Idx != 0)
-			{
-				Objects.RemoveAt(0, Idx);
-			}
-			return Objects.Num() != 0;
-		}
-
-		void Flush()
-		{
-			// See if all samples are ready to be retired now...
-			if (!Update())
-			{
-				// They are. No need for any async task...
-				return;
-			}
-
-			// Some samples still need the GPU to get done. Use async task to get this done...
-			TFunction<void()> FlushTask = [LastObjects{ MoveTemp(Objects) }]()
-			{
-				while (1)
-				{
-#if USE_LIMITED_FENCEWAIT
-					double Now = FPlatformTime::Seconds();
-#endif
-					int32 Idx = 0;
-					for (; Idx < LastObjects.Num(); ++Idx)
-					{
-						// Still not signaled?
-						if (LastObjects[Idx].GPUFence.IsValid() && !LastObjects[Idx].GPUFence->Poll())
-						{
-#if USE_LIMITED_FENCEWAIT
-							// HACK: But how long has this been going on? Might we have a fence that never will signal?
-							if ((Now - LastObjects[Idx].RetireTime) < MaxWaitForFence)
-#else
-							if (1)
-#endif
-							{
-								break;
-							}
-						}
-					}
-					if (Idx == LastObjects.Num())
-					{
-						break;
-					}
-
-					FPlatformProcess::Sleep(5.0f / 1000.0f);
-				}
-			};
-			Async(EAsyncExecution::ThreadPool, MoveTemp(FlushTask));
-		}
-
-		struct FRetiringObjectInfo
-		{
-			ObjectRefType Object;
-			FGPUFenceRHIRef GPUFence;
-			double RetireTime;
-		};
-
-		TArray<FRetiringObjectInfo> Objects;
-		FCriticalSection CS;
-	};
-
-	typedef TGPUsyncedDataDeleter<TSharedPtr<IMediaTextureSample, ESPMode::ThreadSafe>> FPriorSamples;
-
 	TSharedRef<FPriorSamples, ESPMode::ThreadSafe> PriorSamples;
 
 	/** cached params etc. for use with mip generator */
