Various fixes for Unity Builds

diff --git a/Engine/Source/Runtime/AudioCodecEngine/Public/DecoderInputFactory.h b/Engine/Source/Runtime/AudioCodecEngine/Public/DecoderInputFactory.h
--- a/Engine/Source/Runtime/AudioCodecEngine/Public/DecoderInputFactory.h
+++ b/Engine/Source/Runtime/AudioCodecEngine/Public/DecoderInputFactory.h
@@ -3,11 +3,13 @@
 #pragma once
 
 #include "Templates/UniquePtr.h"
+#include "AudioMixer.h"
 
 // Forward declares.
 class USoundWave;
 class FSoundWaveProxy;
 
+
 namespace Audio
 {
 	// Forward declares.
diff --git a/Engine/Source/Runtime/MediaAssets/Private/Misc/MediaTextureResource.cpp b/Engine/Source/Runtime/MediaAssets/Private/Misc/MediaTextureResource.cpp
--- a/Engine/Source/Runtime/MediaAssets/Private/Misc/MediaTextureResource.cpp
+++ b/Engine/Source/Runtime/MediaAssets/Private/Misc/MediaTextureResource.cpp
@@ -2,7 +2,7 @@
 
 #include "Misc/MediaTextureResource.h"
 #include "MediaAssetsPrivate.h"
-
+#include "Modules/ModuleManager.h"
 #include "DeviceProfiles/DeviceProfile.h"
 #include "DeviceProfiles/DeviceProfileManager.h"
 #include "ExternalTexture.h"
