MediaFramework: Fixed playlist tests failing.
#rnx

diff --git a/Engine/Source/Runtime/MediaAssets/Private/Tests/MediaPlaylistTest.cpp b/Engine/Source/Runtime/MediaAssets/Private/Tests/MediaPlaylistTest.cpp
--- a/Engine/Source/Runtime/MediaAssets/Private/Tests/MediaPlaylistTest.cpp
+++ b/Engine/Source/Runtime/MediaAssets/Private/Tests/MediaPlaylistTest.cpp
@@ -1,6 +1,7 @@
 // Copyright Epic Games, Inc. All Rights Reserved.
 
 #include "CoreTypes.h"
+#include "FileMediaSource.h"
 #include "MediaPlaylist.h"
 #include "Misc/AutomationTest.h"
 #include "UObject/UObjectGlobals.h"
@@ -43,8 +44,9 @@ bool FMediaPlaylistTest::RunTest(const FString& Parameters)
 	}
 
 	// one entry
+	UMediaSource* MediaSource1 = NewObject<UFileMediaSource>();
 	{
-		Playlist->Add((UMediaSource*)1);
+		Playlist->Add(MediaSource1);
 		TestEqual(TEXT("A play list with one entry and must have length 1"), Playlist->Num(), 1);
 
 		Index = INDEX_NONE;
@@ -73,32 +75,33 @@ bool FMediaPlaylistTest::RunTest(const FString& Parameters)
 	}
 
 	// two entries
+	UMediaSource* MediaSource2 = NewObject<UFileMediaSource>();
 	{
-		Playlist->Add((UMediaSource*)2);
+		Playlist->Add(MediaSource2);
 		TestEqual(TEXT("A play list with two entries and must have length 2"), Playlist->Num(), 2);
 
 		Index = INDEX_NONE;
-		TestEqual(TEXT("GetNext(INDEX_NONE) on a play list with two entries must return the first item"), Playlist->GetNext(Index), (UMediaSource*)1);
+		TestEqual(TEXT("GetNext(INDEX_NONE) on a play list with two entries must return the first item"), Playlist->GetNext(Index), MediaSource1);
 		TestEqual(TEXT("GetNext(INDEX_NONE) on a play list with two entries must yield 0"), Index, 0);
 
 		Index = INDEX_NONE;
-		TestEqual(TEXT("GetPrevious(INDEX_NONE) on a play list with two entries must return the second item"), Playlist->GetPrevious(Index), (UMediaSource*)2);
+		TestEqual(TEXT("GetPrevious(INDEX_NONE) on a play list with two entries must return the second item"), Playlist->GetPrevious(Index), MediaSource2);
 		TestEqual(TEXT("GetPrevious(INDEX_NONE) on a play list with two entries must yield 1"), Index, 1);
 
 		Index = 0;
-		TestEqual(TEXT("GetNext(0) on a play list with two entries must return the second item"), Playlist->GetNext(Index), (UMediaSource*)2);
+		TestEqual(TEXT("GetNext(0) on a play list with two entries must return the second item"), Playlist->GetNext(Index), MediaSource2);
 		TestEqual(TEXT("GetNext(0) on a play list with two entries must yield 1"), Index, 1);
 
 		Index = 0;
-		TestEqual(TEXT("GetPrevious(0) on a play list with two entries must return the first item"), Playlist->GetPrevious(Index), (UMediaSource*)2);
+		TestEqual(TEXT("GetPrevious(0) on a play list with two entries must return the first item"), Playlist->GetPrevious(Index), MediaSource2);
 		TestEqual(TEXT("GetPrevious(0) on a play list with two entries must yield 1"), Index, 1);
 
 		Index = 1;
-		TestEqual(TEXT("GetNext(1) on a play list with two entries must return the first item"), Playlist->GetNext(Index), (UMediaSource*)1);
+		TestEqual(TEXT("GetNext(1) on a play list with two entries must return the first item"), Playlist->GetNext(Index), MediaSource1);
 		TestEqual(TEXT("GetNext(1) on a play list with two entries must yield 0"), Index, 0);
 
 		Index = 1;
-		TestEqual(TEXT("GetPrevious(1) on a play list with two entries must return the first item"), Playlist->GetPrevious(Index), (UMediaSource*)1);
+		TestEqual(TEXT("GetPrevious(1) on a play list with two entries must return the first item"), Playlist->GetPrevious(Index), MediaSource1);
 		TestEqual(TEXT("GetPrevious(1) on a play list with two entries must yield 0"), Index, 0);
 
 		TestNotNull(TEXT("GetRandom() on a play list with two entries and must not return nullptr"), Playlist->GetRandom(Index));
