Adding guards to prtect against a null playlist instance to MediaPlayer event handling
#rnx

diff --git a/Engine/Source/Runtime/MediaAssets/Private/Assets/MediaPlayer.cpp b/Engine/Source/Runtime/MediaAssets/Private/Assets/MediaPlayer.cpp
--- a/Engine/Source/Runtime/MediaAssets/Private/Assets/MediaPlayer.cpp
+++ b/Engine/Source/Runtime/MediaAssets/Private/Assets/MediaPlayer.cpp
@@ -918,7 +918,7 @@ void UMediaPlayer::HandlePlayerMediaEvent(EMediaEvent Event)
 		}
 		else
 		{
-			PlayerFacade->SetLooping(Loop && (Playlist->Num() == 1));
+			PlayerFacade->SetLooping(Loop && (!Playlist || Playlist->Num() == 1));
 		}
 		PlayerFacade->SetViewField(HorizontalFieldOfView, VerticalFieldOfView, true);
 		PlayerFacade->SetViewOrientation(FQuat(ViewRotation), true);
@@ -944,21 +944,25 @@ void UMediaPlayer::HandlePlayerMediaEvent(EMediaEvent Event)
 	case EMediaEvent::MediaOpenFailed:
 		OnMediaOpenFailed.Broadcast(PlayerFacade->GetUrl());
 
-		if ((Loop && (Playlist->Num() != 1)) || (PlaylistIndex + 1 < Playlist->Num()))
+		if (Playlist)
 		{
-			Next();
+			if ((Loop && (Playlist->Num() != 1)) || (PlaylistIndex + 1 < Playlist->Num()))
+			{
+				Next();
+			}
 		}
 		break;
 
 	case EMediaEvent::PlaybackEndReached:
 		OnEndReached.Broadcast();
 
-		check(Playlist != nullptr);
-
-		if ((Loop && (Playlist->Num() != 1)) || (PlaylistIndex + 1 < Playlist->Num()))
+		if (Playlist)
 		{
-			PlayOnNext = true;
-			Next();
+			if ((Loop && (Playlist->Num() != 1)) || (PlaylistIndex + 1 < Playlist->Num()))
+			{
+				PlayOnNext = true;
+				Next();
+			}
 		}
 		break;
 
