Fix missing AssetManager metadata tags in the assetregistry tags for UBaseMediaSource; need to add a call to the Superclass's function.
#rnx

diff --git a/Engine/Source/Runtime/MediaAssets/Private/Assets/BaseMediaSource.cpp b/Engine/Source/Runtime/MediaAssets/Private/Assets/BaseMediaSource.cpp
--- a/Engine/Source/Runtime/MediaAssets/Private/Assets/BaseMediaSource.cpp
+++ b/Engine/Source/Runtime/MediaAssets/Private/Assets/BaseMediaSource.cpp
@@ -25,15 +25,9 @@ void UBaseMediaSource::GetAssetRegistryTags(TArray<FAssetRegistryTag>& OutTags)
 	{
 		OutTags.Add(FAssetRegistryTag("Url", Url, FAssetRegistryTag::TT_Alphabetical));
 	}
+	Super::GetAssetRegistryTags(OutTags);
 }
 
-
-#if WITH_EDITOR
-void UBaseMediaSource::GetAssetRegistryTagMetadata(TMap<FName, FAssetRegistryTagMetadata>& OutMetadata) const
-{
-}
-#endif
-
 void UBaseMediaSource::PreSave(const class ITargetPlatform* TargetPlatform)
 {
 	PRAGMA_DISABLE_DEPRECATION_WARNINGS;
diff --git a/Engine/Source/Runtime/MediaAssets/Public/BaseMediaSource.h b/Engine/Source/Runtime/MediaAssets/Public/BaseMediaSource.h
--- a/Engine/Source/Runtime/MediaAssets/Public/BaseMediaSource.h
+++ b/Engine/Source/Runtime/MediaAssets/Public/BaseMediaSource.h
@@ -46,10 +46,6 @@ public:
 	virtual void GetAssetRegistryTags(TArray<FAssetRegistryTag>& OutTags) const override;
 	virtual void Serialize(FArchive& Ar) override;
 
-#if WITH_EDITOR
-	virtual void GetAssetRegistryTagMetadata(TMap<FName, FAssetRegistryTagMetadata>& OutMetadata) const override;
-#endif
-
 public:
 
 	//~ IMediaOptions interface
