Fxiing up MediaTextureResource texture sample conversion code to not require conversion for RGB -> BGR changes on 8-bit data anymore (was in there by accident)

diff --git a/Engine/Source/Runtime/MediaAssets/Private/Misc/MediaTextureResource.cpp b/Engine/Source/Runtime/MediaAssets/Private/Misc/MediaTextureResource.cpp
--- a/Engine/Source/Runtime/MediaAssets/Private/Misc/MediaTextureResource.cpp
+++ b/Engine/Source/Runtime/MediaAssets/Private/Misc/MediaTextureResource.cpp
@@ -688,6 +688,7 @@ void FMediaTextureResource::ReleaseDynamicRHI()
 	 const EPixelFormat Format = SampleTexture->GetFormat();
 
 	 return ((Format != PF_B8G8R8A8) &&
+		 (Format != PF_R8G8B8A8) &&
 		 (Format != PF_FloatRGB) &&
 		 (Format != PF_FloatRGBA));
  }
