adding back surround sound capability to mediasoundcomponent

diff --git a/Engine/Source/Runtime/MediaAssets/Private/Assets/MediaSoundComponent.cpp b/Engine/Source/Runtime/MediaAssets/Private/Assets/MediaSoundComponent.cpp
--- a/Engine/Source/Runtime/MediaAssets/Private/Assets/MediaSoundComponent.cpp
+++ b/Engine/Source/Runtime/MediaAssets/Private/Assets/MediaSoundComponent.cpp
@@ -329,14 +329,14 @@ bool UMediaSoundComponent::Init(int32& SampleRate)
 	{
 		NumChannels = 1;
 	}
-	else //if (Channels == EMediaSoundChannels::Stereo)
+	else if (Channels == EMediaSoundChannels::Stereo)
 	{
 		NumChannels = 2;
-	}/*
+	}
 	else
 	{
 		NumChannels = 8;
-	}*/
+	}
 
 	// increase buffer callback size for media decoding. Media doesn't need fast response time so can decode more per callback.
 	//PreferredBufferLength = NumChannels * 8196;
