MediaTexture: Fix for auto clear not working.
#rnx

diff --git a/Engine/Source/Runtime/MediaAssets/Private/Assets/MediaTexture.cpp b/Engine/Source/Runtime/MediaAssets/Private/Assets/MediaTexture.cpp
--- a/Engine/Source/Runtime/MediaAssets/Private/Assets/MediaTexture.cpp
+++ b/Engine/Source/Runtime/MediaAssets/Private/Assets/MediaTexture.cpp
@@ -83,6 +83,7 @@ UMediaTexture::UMediaTexture(const FObjectInitializer& ObjectInitializer)
 	, CurrentOrientation(MTORI_Original)
 	, DefaultGuid(FGuid::NewGuid())
 	, Dimensions(FIntPoint::ZeroValue)
+	, bIsCleared(false)
 	, Size(0)
 	, CachedNextSampleTime(FTimespan::MinValue())
 	, TextureNumMips(1)
@@ -340,7 +341,7 @@ void UMediaTexture::TickResource(FTimespan Timecode)
 
 	if (!CurrentPlayer.IsValid())
 	{
-		if ((LastClearColor == ClearColor) && (LastSrgb == SRGB))
+		if ((LastClearColor == ClearColor) && (LastSrgb == SRGB) && (bIsCleared))
 		{
 			return; // nothing to render
 		}
@@ -352,6 +353,7 @@ void UMediaTexture::TickResource(FTimespan Timecode)
 	// set up render parameters
 	FMediaTextureResource::FRenderParams RenderParams;
 
+	bool bIsSampleValid = false;
 	if (UMediaPlayer* CurrentPlayerPtr = CurrentPlayer.Get())
 	{
 		const bool PlayerActive = CurrentPlayerPtr->IsPaused() || CurrentPlayerPtr->IsPlaying() || CurrentPlayerPtr->IsPreparing();
@@ -391,6 +393,7 @@ void UMediaTexture::TickResource(FTimespan Timecode)
 				LastSrgb = SRGB;
 
 				TextureNumMips = (Sample->GetNumMips() > 1) ? Sample->GetNumMips() : NumMips;
+				bIsSampleValid = true;
 			}
 			else
 			{
@@ -407,6 +410,7 @@ void UMediaTexture::TickResource(FTimespan Timecode)
 					LastSrgb = SRGB;
 
 					TextureNumMips = (Sample->GetNumMips() > 1) ? Sample->GetNumMips() : NumMips;
+					bIsSampleValid = true;
 				}
 
 				RenderParams.SampleSource = SampleQueue;
@@ -448,6 +452,10 @@ void UMediaTexture::TickResource(FTimespan Timecode)
 		{
 			ResourceParam->Render(RenderParams);
 		});
+
+	
+	// The texture is cleared if we have auto clear enabled, and we do not have a valid sample.
+	bIsCleared = ((AutoClear) && (bIsSampleValid == false));
 }
 
 void UMediaTexture::UpdateSampleInfo(const TSharedPtr<IMediaTextureSample, ESPMode::ThreadSafe> & Sample)
diff --git a/Engine/Source/Runtime/MediaAssets/Public/MediaTexture.h b/Engine/Source/Runtime/MediaAssets/Public/MediaTexture.h
--- a/Engine/Source/Runtime/MediaAssets/Public/MediaTexture.h
+++ b/Engine/Source/Runtime/MediaAssets/Public/MediaTexture.h
@@ -276,6 +276,9 @@ private:
 	/** The previously used sRGB flag. */
 	bool LastSrgb;
 
+	/** True if the texture has been cleared. */
+	bool bIsCleared;
+
 	/** Texture sample queue. */
 	TSharedPtr<FMediaTextureSampleQueue, ESPMode::ThreadSafe> SampleQueue;
 
