MediaTexture: At the moment there is no interpolation between mip levels for external textures, which creates a hard aliased edge. This CL allows filtering between mip levels when mip levels are provided from outside.

diff --git a/Engine/Source/Runtime/MediaAssets/Private/Assets/MediaTexture.cpp b/Engine/Source/Runtime/MediaAssets/Private/Assets/MediaTexture.cpp
--- a/Engine/Source/Runtime/MediaAssets/Private/Assets/MediaTexture.cpp
+++ b/Engine/Source/Runtime/MediaAssets/Private/Assets/MediaTexture.cpp
@@ -168,7 +168,7 @@ FTextureResource* UMediaTexture::CreateResource()
 		}
 	}
 
-	Filter = (EnableGenMips && (TextureNumMips > 1)) ? TF_Trilinear : TF_Bilinear;
+	Filter = (TextureNumMips > 1) ? TF_Trilinear : TF_Bilinear;
 
 	return new FMediaTextureResource(*this, Dimensions, Size, ClearColor, CurrentGuid.IsValid() ? CurrentGuid : DefaultGuid, EnableGenMips, NumMips);
 }
@@ -436,7 +436,7 @@ void UMediaTexture::TickResource(FTimespan Timecode)
 	}
 
 	// update filter state, responding to mips setting
-	Filter = (EnableGenMips && (TextureNumMips > 1)) ? TF_Trilinear : TF_Bilinear;
+	Filter = (TextureNumMips > 1) ? TF_Trilinear : TF_Bilinear;
 
 	// setup render parameters
 	RenderParams.CanClear = AutoClear;
diff --git a/Engine/Source/Runtime/MediaAssets/Private/Misc/MediaTextureResource.cpp b/Engine/Source/Runtime/MediaAssets/Private/Misc/MediaTextureResource.cpp
--- a/Engine/Source/Runtime/MediaAssets/Private/Misc/MediaTextureResource.cpp
+++ b/Engine/Source/Runtime/MediaAssets/Private/Misc/MediaTextureResource.cpp
@@ -575,7 +575,7 @@ uint32 FMediaTextureResource::GetSizeY() const
 
 void FMediaTextureResource::SetupSampler()
 {
-	ESamplerFilter OwnerFilter = bEnableGenMips ? (ESamplerFilter)UDeviceProfileManager::Get().GetActiveProfile()->GetTextureLODSettings()->GetSamplerFilter(&Owner) : SF_Bilinear;
+	ESamplerFilter OwnerFilter = (CurrentNumMips > 1) ? (ESamplerFilter)UDeviceProfileManager::Get().GetActiveProfile()->GetTextureLODSettings()->GetSamplerFilter(&Owner) : SF_Bilinear;
 
 	if (CurrentSamplerFilter != OwnerFilter)
 	{
