MediaAssets: ImgMediaEngine texture tracking moved to MediaTextureTracker so texture tracking can be done without the ImgMedia plugin.

#rnx

diff --git a/Engine/Plugins/Media/ImgMedia/Source/ImgMedia/Private/Assets/ImgMediaMipMapInfo.cpp b/Engine/Plugins/Media/ImgMedia/Source/ImgMedia/Private/Assets/ImgMediaMipMapInfo.cpp
--- a/Engine/Plugins/Media/ImgMedia/Source/ImgMedia/Private/Assets/ImgMediaMipMapInfo.cpp
+++ b/Engine/Plugins/Media/ImgMedia/Source/ImgMedia/Private/Assets/ImgMediaMipMapInfo.cpp
@@ -2,7 +2,6 @@
 
 #include "ImgMediaMipMapInfo.h"
 
-#include "ImgMediaEngine.h"
 #include "ImgMediaMipMapInfoManager.h"
 #include "ImgMediaPrivate.h"
 
@@ -10,6 +9,7 @@
 #include "Engine/Engine.h"
 #include "Engine/StaticMesh.h"
 #include "GameFramework/Actor.h"
+#include "MediaTextureTracker.h"
 
 DECLARE_CYCLE_STAT(TEXT("ImgMedia MipMap Update Cache"), STAT_ImgMedia_MipMapUpdateCache, STATGROUP_Media);
 
@@ -80,19 +80,19 @@ void FImgMediaMipMapInfo::AddObjectsUsingThisMediaTexture(UMediaTexture* InMedia
 {
 	// Get objects using this texture.
 	FImgMediaMipMapInfoManager& InfoManager = FImgMediaMipMapInfoManager::Get();
-	FImgMediaEngine& ImgMediaEngine = FImgMediaEngine::Get();
-	const TArray<TWeakPtr<FImgMediaMipMapObjectInfo, ESPMode::ThreadSafe>>* ObjectInfos = ImgMediaEngine.GetObjects(InMediaTexture);
+	FMediaTextureTracker& TextureTracker = FMediaTextureTracker::Get();
+	const TArray<TWeakPtr<FMediaTextureTrackerObject, ESPMode::ThreadSafe>>* ObjectInfos = TextureTracker.GetObjects(InMediaTexture);
 	if (ObjectInfos != nullptr)
 	{
-		for (TWeakPtr<FImgMediaMipMapObjectInfo, ESPMode::ThreadSafe> ObjectInfoPtr : *ObjectInfos)
+		for (TWeakPtr<FMediaTextureTrackerObject, ESPMode::ThreadSafe> ObjectInfoPtr : *ObjectInfos)
 		{
-			TSharedPtr<FImgMediaMipMapObjectInfo, ESPMode::ThreadSafe> ObjectInfo = ObjectInfoPtr.Pin();
+			TSharedPtr<FMediaTextureTrackerObject, ESPMode::ThreadSafe> ObjectInfo = ObjectInfoPtr.Pin();
 			if (ObjectInfo.IsValid())
 			{
 				AActor* Owner = ObjectInfo->Object.Get();
 				if (Owner != nullptr)
 				{
-					AddObject(Owner, ObjectInfo->Width, ObjectInfo->LODBias);
+					AddObject(Owner, 0.0f, ObjectInfo->MipMapLODBias);
 				}
 			}
 		}
diff --git a/Engine/Plugins/Media/ImgMedia/Source/ImgMedia/Private/Assets/ImgMediaMipMapInfoManager.cpp b/Engine/Plugins/Media/ImgMedia/Source/ImgMedia/Private/Assets/ImgMediaMipMapInfoManager.cpp
--- a/Engine/Plugins/Media/ImgMedia/Source/ImgMedia/Private/Assets/ImgMediaMipMapInfoManager.cpp
+++ b/Engine/Plugins/Media/ImgMedia/Source/ImgMedia/Private/Assets/ImgMediaMipMapInfoManager.cpp
@@ -6,6 +6,7 @@
 #include "MediaPlayer.h"
 #include "MediaPlayerFacade.h"
 #include "MediaTexture.h"
+#include "MediaTextureTracker.h"
 
 #include "Camera/CameraActor.h"
 #include "Camera/CameraComponent.h"
@@ -96,10 +97,10 @@ void FImgMediaMipMapInfoManager::RemoveCamera(AActor* InActor)
 
 void FImgMediaMipMapInfoManager::GetMediaTexturesFromPlayer(TArray<UMediaTexture*>& OutMediaTextures, IMediaPlayer* InPlayer)
 {
-	FImgMediaEngine& ImgMediaEngine = FImgMediaEngine::Get();
+	FMediaTextureTracker& TextureTracker = FMediaTextureTracker::Get();
 
 	// Look through all the media textures we know about.
-	const TArray<TWeakObjectPtr<UMediaTexture>>& MediaTextures = ImgMediaEngine.GetTextures();
+	const TArray<TWeakObjectPtr<UMediaTexture>>& MediaTextures = TextureTracker.GetTextures();
 	for (TWeakObjectPtr<UMediaTexture> TexturePtr : MediaTextures)
 	{
 		UMediaTexture* Texture = TexturePtr.Get();
diff --git a/Engine/Plugins/Media/ImgMedia/Source/ImgMediaEngine/Private/ImgMediaEngine.cpp b/Engine/Plugins/Media/ImgMedia/Source/ImgMediaEngine/Private/ImgMediaEngine.cpp
--- a/Engine/Plugins/Media/ImgMedia/Source/ImgMediaEngine/Private/ImgMediaEngine.cpp
+++ b/Engine/Plugins/Media/ImgMedia/Source/ImgMediaEngine/Private/ImgMediaEngine.cpp
@@ -16,34 +16,3 @@ FImgMediaEngine& FImgMediaEngine::Get()
 
 	return Engine;
 }
-
-void FImgMediaEngine::RegisterTexture(TSharedPtr<FImgMediaMipMapObjectInfo, ESPMode::ThreadSafe>& InInfo, UMediaTexture* InTexture)
-{
-	// Do we have this media texture yet?
-	TWeakObjectPtr<UMediaTexture> TexturePtr(InTexture);
-	if (MediaTextures.Contains(TexturePtr) == false)
-	{
-		MediaTextures.Emplace(TexturePtr);
-		MapTextureToObject.Emplace(TexturePtr);
-	}
-
-	// Add component to our list.
-	TArray<TWeakPtr<FImgMediaMipMapObjectInfo, ESPMode::ThreadSafe>>* FoundObjects = MapTextureToObject.Find(TexturePtr);
-	FoundObjects->Add(InInfo);
-}
-
-void FImgMediaEngine::UnregisterTexture(TSharedPtr<FImgMediaMipMapObjectInfo, ESPMode::ThreadSafe>& InInfo, UMediaTexture* InTexture)
-{
-	TWeakObjectPtr<UMediaTexture> TexturePtr(InTexture);
-	TArray<TWeakPtr<FImgMediaMipMapObjectInfo, ESPMode::ThreadSafe>>* FoundObjects = MapTextureToObject.Find(TexturePtr);
-	if (FoundObjects != nullptr)
-	{
-		FoundObjects->RemoveSwap(InInfo);
-	}
-}
-
-const TArray<TWeakPtr<FImgMediaMipMapObjectInfo, ESPMode::ThreadSafe>>* FImgMediaEngine::GetObjects(UMediaTexture* InTexture) const
-{
-	const TArray<TWeakPtr<FImgMediaMipMapObjectInfo, ESPMode::ThreadSafe>>* ObjectsPtr = MapTextureToObject.Find(InTexture);
-	return ObjectsPtr;
-}
diff --git a/Engine/Plugins/Media/ImgMedia/Source/ImgMediaEngine/Private/Unreal/ImgMediaPlaybackComponent.cpp b/Engine/Plugins/Media/ImgMedia/Source/ImgMediaEngine/Private/Unreal/ImgMediaPlaybackComponent.cpp
--- a/Engine/Plugins/Media/ImgMedia/Source/ImgMediaEngine/Private/Unreal/ImgMediaPlaybackComponent.cpp
+++ b/Engine/Plugins/Media/ImgMedia/Source/ImgMediaEngine/Private/Unreal/ImgMediaPlaybackComponent.cpp
@@ -3,10 +3,10 @@
 #include "Unreal/ImgMediaPlaybackComponent.h"
 
 #include "Components/PrimitiveComponent.h"
-#include "ImgMediaEngine.h"
 #include "ImgMediaMipMapInfo.h"
 #include "Materials/MaterialInterface.h"
 #include "MediaTexture.h"
+#include "MediaTextureTracker.h"
 
 UImgMediaPlaybackComponent::UImgMediaPlaybackComponent(const FObjectInitializer& ObjectInitializer)
 	: Super(ObjectInitializer)
@@ -35,11 +35,10 @@ void UImgMediaPlaybackComponent::RegisterWithMipMapInfo()
 	if (Owner != nullptr)
 	{
 		// Create info.
-		FImgMediaEngine& ImgMediaEngine = FImgMediaEngine::Get();
-		ObjectInfo = MakeShared<FImgMediaMipMapObjectInfo, ESPMode::ThreadSafe>();
+		FMediaTextureTracker& MediaTextureTracker = FMediaTextureTracker::Get();
+		ObjectInfo = MakeShared<FMediaTextureTrackerObject, ESPMode::ThreadSafe>();
 		ObjectInfo->Object = Owner;
-		ObjectInfo->Width = Width;
-		ObjectInfo->LODBias = LODBias;
+		ObjectInfo->MipMapLODBias = LODBias;
 
 		// Loop over all primitive components.
 		TArray<UActorComponent*> Components;
@@ -64,7 +63,7 @@ void UImgMediaPlaybackComponent::RegisterWithMipMapInfo()
 						if (MediaTexture != nullptr)
 						{
 							MediaTextures.Add(MediaTexture);
-							ImgMediaEngine.RegisterTexture(ObjectInfo, MediaTexture);
+							MediaTextureTracker.RegisterTexture(ObjectInfo, MediaTexture);
 						}
 					}
 				}
@@ -77,7 +76,7 @@ void UImgMediaPlaybackComponent::UnregisterWithMipMapInfo()
 {
 	if (ObjectInfo != nullptr)
 	{
-		FImgMediaEngine& ImgMediaEngine = FImgMediaEngine::Get();
+		FMediaTextureTracker& MediaTextureTracker = FMediaTextureTracker::Get();
 
 		// Loop over all our media textures.
 		for (TWeakObjectPtr<UMediaTexture> MediaTexturePtr : MediaTextures)
@@ -85,7 +84,7 @@ void UImgMediaPlaybackComponent::UnregisterWithMipMapInfo()
 			UMediaTexture* MediaTexture = MediaTexturePtr.Get();
 			if (MediaTexture != nullptr)
 			{
-				ImgMediaEngine.UnregisterTexture(ObjectInfo, MediaTexture);
+				MediaTextureTracker.UnregisterTexture(ObjectInfo, MediaTexture);
 			}
 		}
 	}
diff --git a/Engine/Plugins/Media/ImgMedia/Source/ImgMediaEngine/Public/ImgMediaEngine.h b/Engine/Plugins/Media/ImgMedia/Source/ImgMediaEngine/Public/ImgMediaEngine.h
--- a/Engine/Plugins/Media/ImgMedia/Source/ImgMediaEngine/Public/ImgMediaEngine.h
+++ b/Engine/Plugins/Media/ImgMedia/Source/ImgMediaEngine/Public/ImgMediaEngine.h
@@ -32,42 +32,6 @@ public:
 	 */
 	IMGMEDIAENGINE_API virtual void GetCameraInfo(TArray<FImgMediaMipMapCameraInfo>& CameraInfos) {}
 
-	/**
-	 * Each object should call this for each media texture that the object has.
-	 *
-	 * @param InInfo Object that has the texture.
-	 * @param InTexture Texture to register.
-	 */
-	IMGMEDIAENGINE_API void RegisterTexture(TSharedPtr<FImgMediaMipMapObjectInfo, ESPMode::ThreadSafe>& InInfo, UMediaTexture* InTexture);
-	
-	/**
-	 * Each object should call this for each media texture that the object has.
-	 *
-	 * @param InInfo Object that has the texture.
-	 * @param InTexture Texture to unregister.
-	 */
-	IMGMEDIAENGINE_API void UnregisterTexture(TSharedPtr<FImgMediaMipMapObjectInfo, ESPMode::ThreadSafe>& InInfo, UMediaTexture* InTexture);
-
-	/**
-	 * Get which objects are usnig a specific media texture.
-	 *
-	 * @param InTexture Media texture to check.
-	 * @return Nullptr if nothing found, or a pointer to an array of objects.
-	 */
-	IMGMEDIAENGINE_API const TArray<TWeakPtr<FImgMediaMipMapObjectInfo, ESPMode::ThreadSafe>>* GetObjects(UMediaTexture* InTexture) const;
-	
-	/**
-	 * Get a list of media textures we know about.
-	 *
-	 * @ return List of textures.
-	 */
-	IMGMEDIAENGINE_API const TArray<TWeakObjectPtr<UMediaTexture>>& GetTextures() { return MediaTextures; }
-
-protected:
-	/** Maps a media texture to an array of playback components. */
-	TMap<TWeakObjectPtr<UMediaTexture>, TArray<TWeakPtr<FImgMediaMipMapObjectInfo, ESPMode::ThreadSafe>>> MapTextureToObject;
-	/** Array of media textures that we know about. */
-	TArray<TWeakObjectPtr<UMediaTexture>> MediaTextures;
 };
 
 
diff --git a/Engine/Plugins/Media/ImgMedia/Source/ImgMediaEngine/Public/Unreal/ImgMediaPlaybackComponent.h b/Engine/Plugins/Media/ImgMedia/Source/ImgMediaEngine/Public/Unreal/ImgMediaPlaybackComponent.h
--- a/Engine/Plugins/Media/ImgMedia/Source/ImgMediaEngine/Public/Unreal/ImgMediaPlaybackComponent.h
+++ b/Engine/Plugins/Media/ImgMedia/Source/ImgMediaEngine/Public/Unreal/ImgMediaPlaybackComponent.h
@@ -6,7 +6,7 @@
 #include "ImgMediaPlaybackComponent.generated.h"
 
 class UMediaTexture;
-struct FImgMediaMipMapObjectInfo;
+struct FMediaTextureTrackerObject;
 
 /**
  * Component to help with ImgMedia playback.
@@ -21,9 +21,6 @@ class IMGMEDIAENGINE_API UImgMediaPlaybackComponent : public UActorComponent
 public:
 	UImgMediaPlaybackComponent(const FObjectInitializer& ObjectInitializer);
 
-	/** Width of the object. If < 0, then get the width automatically. */
-	UPROPERTY(EditAnywhere, Category = MipMaps)
-	float Width = -1.0f;
 	/** This will be added to the calculated mipmap level. E.g. if set to 2, and you would normally be at mipmap level 1, then you will actually be at level 3. */
 	UPROPERTY(EditAnywhere, Category = MipMaps)
 	float LODBias = 0.0f;
@@ -47,5 +44,5 @@ protected:
 	TArray<TWeakObjectPtr<UMediaTexture>> MediaTextures;
 
 	/** Info representing this object. */
-	TSharedPtr<FImgMediaMipMapObjectInfo, ESPMode::ThreadSafe> ObjectInfo;
+	TSharedPtr<FMediaTextureTrackerObject, ESPMode::ThreadSafe> ObjectInfo;
 };
diff --git a/Engine/Source/Runtime/MediaAssets/Private/Assets/MediaTextureTracker.cpp b/Engine/Source/Runtime/MediaAssets/Private/Assets/MediaTextureTracker.cpp
new file mode 100644
--- /dev/null
+++ b/Engine/Source/Runtime/MediaAssets/Private/Assets/MediaTextureTracker.cpp
@@ -0,0 +1,43 @@
+// Copyright Epic Games, Inc. All Rights Reserved.
+
+#include "MediaTextureTracker.h"
+#include "MediaTexture.h"
+
+FMediaTextureTracker& FMediaTextureTracker::Get()
+{
+	static FMediaTextureTracker Engine;
+	return Engine;
+}
+
+void FMediaTextureTracker::RegisterTexture(TSharedPtr<FMediaTextureTrackerObject, ESPMode::ThreadSafe>& InInfo,
+	TObjectPtr<UMediaTexture> InTexture)
+{
+	// Do we have this media texture yet?
+	TWeakObjectPtr<UMediaTexture> TexturePtr(InTexture);
+	if (MediaTextures.Contains(TexturePtr) == false)
+	{
+		MediaTextures.Emplace(TexturePtr);
+		MapTextureToObject.Emplace(TexturePtr);
+	}
+
+	// Add component to our list.
+	TArray<TWeakPtr<FMediaTextureTrackerObject, ESPMode::ThreadSafe>>* FoundObjects = MapTextureToObject.Find(TexturePtr);
+	FoundObjects->Add(InInfo);
+}
+
+void FMediaTextureTracker::UnregisterTexture(TSharedPtr<FMediaTextureTrackerObject, ESPMode::ThreadSafe>& InInfo,
+	TObjectPtr<UMediaTexture> InTexture)
+{
+	TWeakObjectPtr<UMediaTexture> TexturePtr(InTexture);
+	TArray<TWeakPtr<FMediaTextureTrackerObject, ESPMode::ThreadSafe>>* FoundObjects = MapTextureToObject.Find(TexturePtr);
+	if (FoundObjects != nullptr)
+	{
+		FoundObjects->RemoveSwap(InInfo);
+	}
+}
+
+const TArray<TWeakPtr<FMediaTextureTrackerObject, ESPMode::ThreadSafe>>* FMediaTextureTracker::GetObjects(TObjectPtr<UMediaTexture> InTexture) const
+{
+	const TArray<TWeakPtr<FMediaTextureTrackerObject, ESPMode::ThreadSafe>>* ObjectsPtr = MapTextureToObject.Find(InTexture);
+	return ObjectsPtr;
+}
diff --git a/Engine/Source/Runtime/MediaAssets/Public/MediaTextureTracker.h b/Engine/Source/Runtime/MediaAssets/Public/MediaTextureTracker.h
new file mode 100644
--- /dev/null
+++ b/Engine/Source/Runtime/MediaAssets/Public/MediaTextureTracker.h
@@ -0,0 +1,73 @@
+// Copyright Epic Games, Inc. All Rights Reserved.
+
+#pragma once
+
+#include "CoreMinimal.h"
+#include "UObject/ObjectPtr.h"
+#include "UObject/WeakObjectPtr.h"
+#include "UObject/WeakObjectPtrTemplates.h"
+
+class AActor;
+class UMediaTexture;
+
+/** Holds info on a single object. */
+struct FMediaTextureTrackerObject
+{
+	/** Actor that is using our img sequence. */
+	TWeakObjectPtr<class AActor> Object;
+	/** LOD bias for the mipmap level. */
+	float MipMapLODBias;
+};
+
+/**
+ * Tracks which media textures are used by which objects.
+ */
+class FMediaTextureTracker
+{
+public:
+	/**
+	 * Access engine from here.
+	 *
+	 * @return Engine.
+	 */
+	MEDIAASSETS_API static FMediaTextureTracker& Get();
+
+	virtual ~FMediaTextureTracker() {}
+
+	/**
+	 * Each object should call this for each media texture that the object has.
+	 *
+	 * @param InInfo Object that has the texture.
+	 * @param InTexture Texture to register.
+	 */
+	MEDIAASSETS_API void RegisterTexture(TSharedPtr<FMediaTextureTrackerObject, ESPMode::ThreadSafe>& InInfo, TObjectPtr<UMediaTexture> InTexture);
+
+	/**
+	 * Each object should call this for each media texture that the object has.
+	 *
+	 * @param InInfo Object that has the texture.
+	 * @param InTexture Texture to unregister.
+	 */
+	MEDIAASSETS_API void UnregisterTexture(TSharedPtr<FMediaTextureTrackerObject, ESPMode::ThreadSafe>& InInfo, TObjectPtr<UMediaTexture> InTexture);
+
+	/**
+	 * Get which objects are usnig a specific media texture.
+	 *
+	 * @param InTexture Media texture to check.
+	 * @return Nullptr if nothing found, or a pointer to an array of objects.
+	 */
+	MEDIAASSETS_API const TArray<TWeakPtr<FMediaTextureTrackerObject, ESPMode::ThreadSafe>>* GetObjects(TObjectPtr<UMediaTexture> InTexture) const;
+
+	/**
+	 * Get a list of media textures we know about.
+	 *
+	 * @ return List of textures.
+	 */
+	MEDIAASSETS_API const TArray<TWeakObjectPtr<UMediaTexture>>& GetTextures() { return MediaTextures; }
+
+protected:
+	/** Maps a media texture to an array of playback components. */
+	TMap<TWeakObjectPtr<UMediaTexture>, TArray<TWeakPtr<FMediaTextureTrackerObject, ESPMode::ThreadSafe>>> MapTextureToObject;
+	/** Array of media textures that we know about. */
+	TArray<TWeakObjectPtr<UMediaTexture>> MediaTextures;
+};
