MediaAssets: Media sources can be associated with a file extension.

#rnx

diff --git a/Engine/Source/Runtime/MediaAssets/Private/Assets/MediaSource.cpp b/Engine/Source/Runtime/MediaAssets/Private/Assets/MediaSource.cpp
--- a/Engine/Source/Runtime/MediaAssets/Private/Assets/MediaSource.cpp
+++ b/Engine/Source/Runtime/MediaAssets/Private/Assets/MediaSource.cpp
@@ -2,6 +2,48 @@
 
 #include "MediaSource.h"
 
+#include "Misc/Paths.h"
+
+void UMediaSource::RegisterSpawnFromFileExtension(const FString& Extension,
+	FMediaSourceSpawnDelegate InDelegate)
+{
+	TMap<FString, FMediaSourceSpawnDelegate>& Delegates =
+		GetSpawnFromFileExtensionDelegates();
+
+	Delegates.Emplace(Extension, InDelegate);
+}
+
+void UMediaSource::UnregisterSpawnFromFileExtension(const FString& Extension)
+{
+	TMap<FString, FMediaSourceSpawnDelegate>& Delegates =
+		GetSpawnFromFileExtensionDelegates();
+
+	Delegates.Remove(Extension);
+}
+
+UMediaSource* UMediaSource::SpawnMediaSourceForString(const FString& MediaPath)
+{
+	TObjectPtr<UMediaSource> MediaSource = nullptr;
+
+	// Do we know about this file extension?
+	FString FileExtension = FPaths::GetExtension(MediaPath);
+	TMap<FString, FMediaSourceSpawnDelegate>& Delegates =
+		GetSpawnFromFileExtensionDelegates();
+	FMediaSourceSpawnDelegate* Delegate = Delegates.Find(FileExtension);
+	if ((Delegate != nullptr) && (Delegate->IsBound()))
+	{
+		MediaSource = Delegate->Execute(MediaPath);
+	}
+
+	return MediaSource;
+}
+
+TMap<FString, FMediaSourceSpawnDelegate>& UMediaSource::GetSpawnFromFileExtensionDelegates()
+{
+	static TMap<FString, FMediaSourceSpawnDelegate> Delegates;
+	return Delegates;
+}
+
 /* IMediaOptions interface
  *****************************************************************************/
 
diff --git a/Engine/Source/Runtime/MediaAssets/Public/MediaSource.h b/Engine/Source/Runtime/MediaAssets/Public/MediaSource.h
--- a/Engine/Source/Runtime/MediaAssets/Public/MediaSource.h
+++ b/Engine/Source/Runtime/MediaAssets/Public/MediaSource.h
@@ -12,6 +12,8 @@
 
 #include "MediaSource.generated.h"
 
+/** Delegate for creating a media source from a string. */
+DECLARE_DELEGATE_RetVal_OneParam(UMediaSource*, FMediaSourceSpawnDelegate, const FString&);
 
 /**
  * Abstract base class for media sources.
@@ -48,6 +50,31 @@ public:
 	UFUNCTION(BlueprintCallable, Category="Media|MediaSource")
 	virtual bool Validate() const PURE_VIRTUAL(UMediaSource::Validate, return false;);
 
+	/**
+	 * Call this to register a callback when someone calls SpawnMediaSourceForString.
+	 * This lets you spawn a media source if the file extension matches what you want.
+	 * 
+	 * @param Extension		File extension to match. This is case insensitive.
+	 * @param InDelegate	This will get called if the Url passed into GetMediaSourceForUrl
+	 *						matches Extension.
+	 */
+	static void RegisterSpawnFromFileExtension(const FString& Extension, FMediaSourceSpawnDelegate InDelegate);
+	
+	/**
+	 * Call this to unregister a callback set with RegisterSpawnFromFileExtension.
+	 *
+	 * @param Extension		File extension that the callack was registered with.
+	 */
+	static void UnregisterSpawnFromFileExtension(const FString& Extension);
+
+	/**
+	 * Call this to try and create a media source appropriate for the media.
+	 *
+	 * @param MediaPath		Can be a file location or a Url.
+	 * @return				Media source or nullptr if none are appropriate.
+	 */
+	static UMediaSource* SpawnMediaSourceForString(const FString& MediaPath);
+
 public:
 
 	//~ IMediaOptions interface
@@ -90,4 +117,9 @@ private:
 	 * Sets the media option specified by Key to the supplied Variant.
 	 */
 	void SetMediaOption(const FName& Key, FVariant& Value);
+
+	/**
+	 * Get a mapping of file extensions to spawn delegates.
+	 */
+	static TMap<FString, FMediaSourceSpawnDelegate>& GetSpawnFromFileExtensionDelegates();
 };
