MediaAssets: Added support for urls for SpawnMediaSourceForString.
#rnx

diff --git a/Engine/Source/Runtime/MediaAssets/Private/Assets/MediaSource.cpp b/Engine/Source/Runtime/MediaAssets/Private/Assets/MediaSource.cpp
--- a/Engine/Source/Runtime/MediaAssets/Private/Assets/MediaSource.cpp
+++ b/Engine/Source/Runtime/MediaAssets/Private/Assets/MediaSource.cpp
@@ -2,7 +2,9 @@
 
 #include "MediaSource.h"
 
+#include "MediaAssetsPrivate.h"
 #include "Misc/Paths.h"
+#include "StreamMediaSource.h"
 
 void UMediaSource::RegisterSpawnFromFileExtension(const FString& Extension,
 	FMediaSourceSpawnDelegate InDelegate)
@@ -25,14 +27,35 @@ UMediaSource* UMediaSource::SpawnMediaSourceForString(const FString& MediaPath)
 {
 	TObjectPtr<UMediaSource> MediaSource = nullptr;
 
-	// Do we know about this file extension?
-	FString FileExtension = FPaths::GetExtension(MediaPath);
-	TMap<FString, FMediaSourceSpawnDelegate>& Delegates =
-		GetSpawnFromFileExtensionDelegates();
-	FMediaSourceSpawnDelegate* Delegate = Delegates.Find(FileExtension);
-	if ((Delegate != nullptr) && (Delegate->IsBound()))
+	// Is it a URL?
+	bool bIsUrl = MediaPath.Contains(TEXT("://"));
+	if (bIsUrl)
+	{
+		TObjectPtr<UStreamMediaSource> StreamMediaSource = NewObject<UStreamMediaSource>(GetTransientPackage(), NAME_None, RF_Transactional | RF_Transient);
+		StreamMediaSource->StreamUrl = MediaPath;
+		MediaSource = StreamMediaSource;
+	}
+	else
+	{
+		// Do we know about this file extension?
+		FString FileExtension = FPaths::GetExtension(MediaPath);
+		TMap<FString, FMediaSourceSpawnDelegate>& Delegates =
+			GetSpawnFromFileExtensionDelegates();
+		FMediaSourceSpawnDelegate* Delegate = Delegates.Find(FileExtension);
+		if ((Delegate != nullptr) && (Delegate->IsBound()))
+		{
+			MediaSource = Delegate->Execute(MediaPath);
+		}
+	}
+
+	// Validate the media source.
+	if (MediaSource != nullptr)
 	{
-		MediaSource = Delegate->Execute(MediaPath);
+		if (!MediaSource->Validate())
+		{
+			UE_LOG(LogMediaAssets, Error, TEXT("Failed to validate %s"), *MediaPath);
+			MediaSource = nullptr;
+		}
 	}
 
 	return MediaSource;
