MediaAssets: Added support for file media sources to SpawnMediaSourceForString.
#rnx

diff --git a/Engine/Source/Runtime/MediaAssets/Private/Assets/MediaSource.cpp b/Engine/Source/Runtime/MediaAssets/Private/Assets/MediaSource.cpp
--- a/Engine/Source/Runtime/MediaAssets/Private/Assets/MediaSource.cpp
+++ b/Engine/Source/Runtime/MediaAssets/Private/Assets/MediaSource.cpp
@@ -2,6 +2,7 @@
 
 #include "MediaSource.h"
 
+#include "FileMediaSource.h"
 #include "MediaAssetsPrivate.h"
 #include "Misc/Paths.h"
 #include "StreamMediaSource.h"
@@ -46,6 +47,13 @@ UMediaSource* UMediaSource::SpawnMediaSourceForString(const FString& MediaPath)
 		{
 			MediaSource = Delegate->Execute(MediaPath);
 		}
+		else
+		{
+			// Try a file media source.
+			TObjectPtr<UFileMediaSource> FileMediaSource = NewObject<UFileMediaSource>(GetTransientPackage(), NAME_None, RF_Transactional | RF_Transient);
+			FileMediaSource->SetFilePath(MediaPath);
+			MediaSource = FileMediaSource;
+		}
 	}
 
 	// Validate the media source.
