MediaCompositing: No longer needs media plate plugin to get the media player for object bindings.

#rnx

diff --git a/Engine/Plugins/Media/MediaCompositing/Source/MediaCompositing/MediaCompositing.Build.cs b/Engine/Plugins/Media/MediaCompositing/Source/MediaCompositing/MediaCompositing.Build.cs
--- a/Engine/Plugins/Media/MediaCompositing/Source/MediaCompositing/MediaCompositing.Build.cs
+++ b/Engine/Plugins/Media/MediaCompositing/Source/MediaCompositing/MediaCompositing.Build.cs
@@ -32,15 +32,6 @@ namespace UnrealBuildTool.Rules
 					"MediaCompositing/Private",
 					"MediaCompositing/Private/MovieScene",
 				});
-
-			DynamicallyLoadedModuleNames.AddRange(
-				new string[] {
-					"MediaPlate",
-				});
-			PrivateIncludePathModuleNames.AddRange(
-				new string[] {
-					"MediaPlate",
-				});
 		}
 	}
 }
diff --git a/Engine/Plugins/Media/MediaCompositing/Source/MediaCompositing/Private/MovieScene/MovieSceneMediaTemplate.cpp b/Engine/Plugins/Media/MediaCompositing/Source/MediaCompositing/Private/MovieScene/MovieSceneMediaTemplate.cpp
--- a/Engine/Plugins/Media/MediaCompositing/Source/MediaCompositing/Private/MovieScene/MovieSceneMediaTemplate.cpp
+++ b/Engine/Plugins/Media/MediaCompositing/Source/MediaCompositing/Private/MovieScene/MovieSceneMediaTemplate.cpp
@@ -2,8 +2,8 @@
 
 #include "MovieSceneMediaTemplate.h"
 
+#include "IMediaAssetsModule.h"
 #include "Math/UnrealMathUtility.h"
-#include "MediaPlateModule.h"
 #include "MediaPlayer.h"
 #include "MediaPlayerFacade.h"
 #include "MediaSoundComponent.h"
@@ -287,15 +287,15 @@ void FMovieSceneMediaSectionTemplate::Initialize(const FMovieSceneEvaluationOper
 			if (Operand.ObjectBindingID.IsValid())
 			{
 				// Yes. Get the media player from the object.
-				FMediaPlateModule* MediaPlateModule = FModuleManager::LoadModulePtr<FMediaPlateModule>("MediaPlate");
-				if (MediaPlateModule != nullptr)
+				IMediaAssetsModule* MediaAssetsModule = FModuleManager::LoadModulePtr<IMediaAssetsModule>("MediaAssets");
+				if (MediaAssetsModule != nullptr)
 				{
 					for (TWeakObjectPtr<> WeakObject : Player.FindBoundObjects(Operand))
 					{
 						UObject* BoundObject = WeakObject.Get();
 						if (BoundObject != nullptr)
 						{
-							MediaPlayer = MediaPlateModule->GetMediaPlayer(BoundObject);
+							MediaPlayer = MediaAssetsModule->GetPlayerFromObject(BoundObject);
 							break;
 						}
 					}
diff --git a/Engine/Plugins/Media/MediaPlate/Source/MediaPlate/Private/MediaPlateModule.cpp b/Engine/Plugins/Media/MediaPlate/Source/MediaPlate/Private/MediaPlateModule.cpp
--- a/Engine/Plugins/Media/MediaPlate/Source/MediaPlate/Private/MediaPlateModule.cpp
+++ b/Engine/Plugins/Media/MediaPlate/Source/MediaPlate/Private/MediaPlateModule.cpp
@@ -2,20 +2,17 @@
 
 #include "MediaPlateModule.h"
 
+#include "IMediaAssetsModule.h"
 #include "MediaPlate.h"
 #include "MediaPlateComponent.h"
 
 DEFINE_LOG_CATEGORY(LogMediaPlate);
 
-UClass* FMediaPlateModule::GetAMediaPlateClass()
-{
-	return AMediaPlate::StaticClass();
-}
-
 UMediaPlayer* FMediaPlateModule::GetMediaPlayer(UObject* Object)
 {
 	UMediaPlayer* MediaPlayer = nullptr;
 
+	// Is this a media plate?
 	AMediaPlate* MediaPlate = Cast<AMediaPlate>(Object);
 	if (MediaPlate != nullptr)
 	{
@@ -31,10 +28,24 @@ UMediaPlayer* FMediaPlateModule::GetMediaPlayer(UObject* Object)
 
 void FMediaPlateModule::StartupModule()
 {
+	IMediaAssetsModule* MediaAssetsModule = FModuleManager::LoadModulePtr<IMediaAssetsModule>("MediaAssets");
+	if (MediaAssetsModule != nullptr)
+	{
+		GetPlayerFromObjectID = MediaAssetsModule->RegisterGetPlayerFromObject(IMediaAssetsModule::FOnGetPlayerFromObject::CreateRaw(this, &FMediaPlateModule::GetMediaPlayer));
+	}
 }
 
 void FMediaPlateModule::ShutdownModule()
 {
+	if (GetPlayerFromObjectID != INDEX_NONE)
+	{
+		IMediaAssetsModule* MediaAssetsModule = FModuleManager::GetModulePtr<IMediaAssetsModule>("MediaAssets");
+		if (MediaAssetsModule != nullptr)
+		{
+			MediaAssetsModule->UnregisterGetPlayerFromObject(GetPlayerFromObjectID);
+			GetPlayerFromObjectID = INDEX_NONE;
+		}
+	}
 }
 
 IMPLEMENT_MODULE(FMediaPlateModule, MediaPlate)
diff --git a/Engine/Plugins/Media/MediaPlate/Source/MediaPlate/Public/MediaPlateModule.h b/Engine/Plugins/Media/MediaPlate/Source/MediaPlate/Public/MediaPlateModule.h
--- a/Engine/Plugins/Media/MediaPlate/Source/MediaPlate/Public/MediaPlateModule.h
+++ b/Engine/Plugins/Media/MediaPlate/Source/MediaPlate/Public/MediaPlateModule.h
@@ -2,11 +2,11 @@
 
 #pragma once
 
-#include "CoreMinimal.h"
 #include "Logging/LogMacros.h"
 #include "Modules/ModuleManager.h"
 
 class UMediaPlayer;
+class UObject;
 
 /** Log category for this module. */
 DECLARE_LOG_CATEGORY_EXTERN(LogMediaPlate, Log, All);
@@ -14,17 +14,16 @@ DECLARE_LOG_CATEGORY_EXTERN(LogMediaPlate, Log, All);
 class MEDIAPLATE_API FMediaPlateModule : public IModuleInterface
 {
 public:
-	/**
-	 * Call this to get the UClass for AMediaPlate.
-	 */
-	virtual UClass* GetAMediaPlateClass();
-
 	/**
 	 * Call this to get the media player from a media plate object.
 	 */
-	virtual UMediaPlayer* GetMediaPlayer(UObject* Object);
+	UMediaPlayer* GetMediaPlayer(UObject* Object);
 
 	/** IModuleInterface implementation */
 	virtual void StartupModule() override;
 	virtual void ShutdownModule() override;
+
+private:
+	/** ID for our delegate. */
+	int32 GetPlayerFromObjectID = INDEX_NONE;
 };
diff --git a/Engine/Source/Runtime/MediaAssets/Private/MediaAssetsModule.cpp b/Engine/Source/Runtime/MediaAssets/Private/MediaAssetsModule.cpp
--- a/Engine/Source/Runtime/MediaAssets/Private/MediaAssetsModule.cpp
+++ b/Engine/Source/Runtime/MediaAssets/Private/MediaAssetsModule.cpp
@@ -3,8 +3,8 @@
 #include "MediaAssetsPrivate.h"
 
 #include "CoreMinimal.h"
+#include "IMediaAssetsModule.h"
 #include "Misc/CoreMisc.h"
-#include "Modules/ModuleInterface.h"
 #include "Modules/ModuleManager.h"
 #include "UObject/UObjectHash.h"
 #include "UObject/UObjectIterator.h"
@@ -20,7 +20,7 @@ DEFINE_LOG_CATEGORY(LogMediaAssets);
  */
 class FMediaAssetsModule
 	: public FSelfRegisteringExec
-	, public IModuleInterface
+	, public IMediaAssetsModule
 {
 public:
 
@@ -56,7 +56,39 @@ public:
 	}
 
 public:
+	//~ IMediaAssetsModule interface
+	virtual int32 RegisterGetPlayerFromObject(const FOnGetPlayerFromObject& Delegate) override
+	{
+		int32 Index = GetPlayerFromObjectDelegates.Num();
+		GetPlayerFromObjectDelegates.Add(Delegate);
+		return Index;
+	}
+
+	virtual void UnregisterGetPlayerFromObject(int32 DelegateID) override
+	{
+		GetPlayerFromObjectDelegates[DelegateID].Unbind();
+	}
 
+	virtual UMediaPlayer* GetPlayerFromObject(UObject* Object) override
+	{
+		UMediaPlayer* Player = nullptr;
+
+		// Go through all our delegates.
+		for (FOnGetPlayerFromObject& Delegate : GetPlayerFromObjectDelegates)
+		{
+			if (Delegate.IsBound())
+			{
+				Player = Delegate.Execute(Object);
+				if (Player != nullptr)
+				{
+					break;
+				}
+			}
+		}
+
+		return Player;
+	}
+	
 	//~ IModuleInterface interface
 
 	virtual void StartupModule() override { }
@@ -66,6 +98,11 @@ public:
 	{
 		return false;
 	}
+
+private:
+	/** All the functions that can get a player from an object. */
+	TArray<FOnGetPlayerFromObject> GetPlayerFromObjectDelegates;
+
 };
 
 
diff --git a/Engine/Source/Runtime/MediaAssets/Public/IMediaAssetsModule.h b/Engine/Source/Runtime/MediaAssets/Public/IMediaAssetsModule.h
new file mode 100644
--- /dev/null
+++ b/Engine/Source/Runtime/MediaAssets/Public/IMediaAssetsModule.h
@@ -0,0 +1,47 @@
+// Copyright Epic Games, Inc. All Rights Reserved.
+
+#pragma once
+
+#include "Delegates/Delegate.h"
+#include "Modules/ModuleInterface.h"
+
+class UMediaPlayer;
+class UObject;
+
+/**
+ * Interface for the MediaAssets module.
+ */
+class IMediaAssetsModule
+	: public IModuleInterface
+{
+public:
+	/** Delegate to get a player from a UObject. */
+	DECLARE_DELEGATE_RetVal_OneParam(UMediaPlayer*, FOnGetPlayerFromObject, UObject*);
+
+	/**
+	 * Plugins should call this so they can provide a function to get a media player from an object.
+	 * 
+	 * @param Delegate		Delegate to get a media player.
+	 * @return ID to pass in to UnregisterGetPlayerFromObject.
+	 */
+	virtual int32 RegisterGetPlayerFromObject(const FOnGetPlayerFromObject& Delegate) = 0;
+
+	/**
+	 * Call this to unregister a delegate.
+	 *
+	 * @param DelegateID	ID returned from RegisterGetPlayerFromObject.
+	 */
+	virtual void UnregisterGetPlayerFromObject(int32 DelegateID) = 0;
+
+	/**
+	 * Call this to get a media player from an object.
+	 * This will query any plugins that have called RegisterGetPlayerFromObject.
+	 * 
+	 * @param Object	Object to get the player from.
+	 * @return Media player, or nullptr if none found. 
+	 */
+	virtual UMediaPlayer* GetPlayerFromObject(UObject* Object) = 0;
+
+	/** Virtual destructor. */
+	virtual ~IMediaAssetsModule() { }
+};
