MediaTextureResource: Added GPU stat for nsights profiling.

diff --git a/Engine/Source/Runtime/MediaAssets/Private/Misc/MediaTextureResource.cpp b/Engine/Source/Runtime/MediaAssets/Private/Misc/MediaTextureResource.cpp
--- a/Engine/Source/Runtime/MediaAssets/Private/Misc/MediaTextureResource.cpp
+++ b/Engine/Source/Runtime/MediaAssets/Private/Misc/MediaTextureResource.cpp
@@ -282,6 +282,7 @@ void FMediaTextureResource::FlushPendingData()
 void FMediaTextureResource::Render(const FRenderParams& Params)
 {
 	check(IsInRenderingThread());
+	SCOPED_GPU_STAT(FRHICommandListExecutor::GetImmediateCommandList(), MediaTextureResource);
 
 	LLM_SCOPE(ELLMTag::MediaStreaming);
 	SCOPE_CYCLE_COUNTER(STAT_MediaAssets_MediaTextureResourceRender);
