- Add support for UAV output configuration from media converter in media texture (used in rivermax/2110 playback)

diff --git a/Engine/Source/Runtime/Media/Public/IMediaTextureSampleConverter.h b/Engine/Source/Runtime/Media/Public/IMediaTextureSampleConverter.h
--- a/Engine/Source/Runtime/Media/Public/IMediaTextureSampleConverter.h
+++ b/Engine/Source/Runtime/Media/Public/IMediaTextureSampleConverter.h
@@ -16,8 +16,9 @@ public:
 	enum EConverterInfoFlags
 	{
 		ConverterInfoFlags_Default = 0,
-		ConverterInfoFlags_WillCreateOutputTexture = 1 << 0,
-		ConverterInfoFlags_PreprocessOnly = 1 << 1,
+		ConverterInfoFlags_WillCreateOutputTexture	= 1 << 0,
+		ConverterInfoFlags_PreprocessOnly			= 1 << 1,
+		ConverterInfoFlags_NeedUAVOutputTexture		= 1 << 2,
 	};
 
 	struct FConversionHints
diff --git a/Engine/Source/Runtime/MediaAssets/Private/Misc/MediaTextureResource.cpp b/Engine/Source/Runtime/MediaAssets/Private/Misc/MediaTextureResource.cpp
--- a/Engine/Source/Runtime/MediaAssets/Private/Misc/MediaTextureResource.cpp
+++ b/Engine/Source/Runtime/MediaAssets/Private/Misc/MediaTextureResource.cpp
@@ -227,6 +227,16 @@ namespace MediaTextureResourceHelpers
 		return Sample->IsOutputSrgb();
 	}
 
+	bool RequiresUAVTexture(const TSharedPtr<IMediaTextureSample, ESPMode::ThreadSafe>& Sample, uint8 NumMips)
+	{	
+		//UAV output is needed if mips are required and uses CS to generate
+		//or if sample converter asks for it
+		const IMediaTextureSampleConverter* Converter = Sample->GetMediaTextureSampleConverter();
+		bool bNeedsUAV = (NumMips > 1 && RHIRequiresComputeGenerateMips());
+		bNeedsUAV |= (Converter && ((Converter->GetConverterInfoFlags() & IMediaTextureSampleConverter::ConverterInfoFlags_NeedUAVOutputTexture) != 0));
+		return bNeedsUAV;
+	}
+
 
 	bool RequiresSrgbInputTexture(const TSharedPtr<IMediaTextureSample, ESPMode::ThreadSafe>& Sample)
 	{
@@ -378,6 +388,8 @@ void FMediaTextureResource::Render(const FRenderParams& Params)
 				IMediaTextureSampleConverter::FConversionHints Hints;
 				Hints.NumMips = (SampleNumMips > 1) ? SampleNumMips : Params.NumMips;
 
+				const bool bNeedsUAVTexture = MediaTextureResourceHelpers::RequiresUAVTexture(Sample, Hints.NumMips);
+
 				// Does the conversion create its own output texture?
 				if ((Converter->GetConverterInfoFlags() & IMediaTextureSampleConverter::ConverterInfoFlags_WillCreateOutputTexture) == 0)
 				{
@@ -395,7 +407,8 @@ void FMediaTextureResource::Render(const FRenderParams& Params)
 					else
 					{
 						// Conversion is fully handled by converter
-						CreateOutputRenderTarget(Sample->GetOutputDim(), MediaTextureResourceHelpers::GetConvertedPixelFormat(Sample), MediaTextureResourceHelpers::RequiresSrgbTexture(Sample), Params.ClearColor, Hints.NumMips);
+
+						CreateOutputRenderTarget(Sample->GetOutputDim(), MediaTextureResourceHelpers::GetConvertedPixelFormat(Sample), MediaTextureResourceHelpers::RequiresSrgbTexture(Sample), Params.ClearColor, Hints.NumMips, bNeedsUAVTexture);
 						Converter->Convert(RenderTargetTextureRHI, Hints);
 					}
 				}
@@ -408,7 +421,7 @@ void FMediaTextureResource::Render(const FRenderParams& Params)
 						// As the converter created the texture, we might need to convert it even more to make it fit our needs. Check...
 						if (RequiresConversion(OutTexture, Sample->GetOutputDim(), NumMips))
 						{
-							CreateOutputRenderTarget(Sample->GetOutputDim(), MediaTextureResourceHelpers::GetConvertedPixelFormat(Sample), MediaTextureResourceHelpers::RequiresSrgbTexture(Sample), Params.ClearColor, Hints.NumMips);
+							CreateOutputRenderTarget(Sample->GetOutputDim(), MediaTextureResourceHelpers::GetConvertedPixelFormat(Sample), MediaTextureResourceHelpers::RequiresSrgbTexture(Sample), Params.ClearColor, Hints.NumMips, bNeedsUAVTexture);
 							ConvertTextureToOutput(OutTexture.GetReference(), Sample);
 						}
 						else
@@ -639,7 +652,9 @@ void FMediaTextureResource::ReleaseDynamicRHI()
  void FMediaTextureResource::ClearTexture(const FLinearColor& ClearColor, bool SrgbOutput)
 {
 	// create output render target if we don't have one yet
-	CreateOutputRenderTarget(FIntPoint(2, 2), PF_B8G8R8A8, SrgbOutput, ClearColor, 1);
+	constexpr uint8 NumMips = 1;
+	constexpr bool bNeedsUAVTexture = false;
+	CreateOutputRenderTarget(FIntPoint(2, 2), PF_B8G8R8A8, SrgbOutput, ClearColor, NumMips, bNeedsUAVTexture);
 
 	// draw the clear color
 	FRHICommandListImmediate& CommandList = FRHICommandListExecutor::GetImmediateCommandList();
@@ -797,7 +812,9 @@ void FMediaTextureResource::ConvertSample(const TSharedPtr<IMediaTextureSample,
 
 	// create the output texture
 	const FIntPoint OutputDim = Sample->GetOutputDim();
-	CreateOutputRenderTarget(OutputDim, MediaTextureResourceHelpers::GetConvertedPixelFormat(Sample), MediaTextureResourceHelpers::RequiresSrgbTexture(Sample), ClearColor, (SampleNumMips > 1) ? SampleNumMips : InNumMips);
+	const uint8 NumMips = (SampleNumMips > 1) ? SampleNumMips : InNumMips;
+	const bool bNeedsUAVTexture = MediaTextureResourceHelpers::RequiresUAVTexture(Sample, NumMips);
+	CreateOutputRenderTarget(OutputDim, MediaTextureResourceHelpers::GetConvertedPixelFormat(Sample), MediaTextureResourceHelpers::RequiresSrgbTexture(Sample), ClearColor, NumMips, bNeedsUAVTexture);
 
 	ConvertTextureToOutput(InputTexture, Sample);
 }
@@ -1017,7 +1034,9 @@ void FMediaTextureResource::CopySample(const TSharedPtr<IMediaTextureSample, ESP
 		else
 		{
 			// Texture to receive texture from sample
-			CreateOutputRenderTarget(Sample->GetOutputDim(), MediaTextureResourceHelpers::GetPixelFormat(Sample), MediaTextureResourceHelpers::RequiresSrgbTexture(Sample), ClearColor, (SampleNumMips > 1) ? SampleNumMips : InNumMips);
+			const uint8 NumMips = (SampleNumMips > 1) ? SampleNumMips : InNumMips;
+			const bool bNeedsUAVTexture = MediaTextureResourceHelpers::RequiresUAVTexture(Sample, NumMips);
+			CreateOutputRenderTarget(Sample->GetOutputDim(), MediaTextureResourceHelpers::GetPixelFormat(Sample), MediaTextureResourceHelpers::RequiresSrgbTexture(Sample), ClearColor, NumMips, bNeedsUAVTexture);
 
 			// Copy data into the output texture to able to add mips later on
 			FRHICommandListExecutor::GetImmediateCommandList().CopyTexture(SampleTexture2D, OutputTarget, FRHICopyTextureInfo());
@@ -1026,7 +1045,9 @@ void FMediaTextureResource::CopySample(const TSharedPtr<IMediaTextureSample, ESP
 	else
 	{
 		// Texture to receive precisely only output pixels via CPU copy
-		CreateOutputRenderTarget(Sample->GetDim(), MediaTextureResourceHelpers::GetPixelFormat(Sample), MediaTextureResourceHelpers::RequiresSrgbTexture(Sample), ClearColor, (SampleNumMips > 1) ? SampleNumMips : InNumMips);
+		const uint8 NumMips = (SampleNumMips > 1) ? SampleNumMips : InNumMips;
+		const bool bNeedsUAVTexture = MediaTextureResourceHelpers::RequiresUAVTexture(Sample, NumMips);
+		CreateOutputRenderTarget(Sample->GetDim(), MediaTextureResourceHelpers::GetPixelFormat(Sample), MediaTextureResourceHelpers::RequiresSrgbTexture(Sample), ClearColor, NumMips, bNeedsUAVTexture);
 
 		// If we also have no source buffer and the platform generally would allow for use of external textures, we assume it is just that...
 		// (as long as the player actually produces (dummy) samples, this will enable mips support as well as auto conversion for "new style output" mode)
@@ -1177,18 +1198,19 @@ void FMediaTextureResource::UpdateTextureReference(FRHITexture2D* NewTexture)
 }
 
 
-void FMediaTextureResource::CreateOutputRenderTarget(const FIntPoint & InDim, EPixelFormat InPixelFormat, bool bInSRGB, const FLinearColor & InClearColor, uint8 InNumMips)
+void FMediaTextureResource::CreateOutputRenderTarget(const FIntPoint & InDim, EPixelFormat InPixelFormat, bool bInSRGB, const FLinearColor & InClearColor, uint8 InNumMips, bool bNeedsUAVSupport)
 {
 	// create output render target if necessary
 	ETextureCreateFlags OutputCreateFlags = TexCreate_Dynamic | (bInSRGB ? TexCreate_SRGB : TexCreate_None);
+	if (bNeedsUAVSupport && RHIIsTypedUAVLoadSupported(InPixelFormat))
+	{
+		OutputCreateFlags |= TexCreate_UAV;
+	}
+
 	if (InNumMips > 1)
 	{
 		// Make sure can have mips & the mip generator has what it needs to work
 		OutputCreateFlags |= TexCreate_GenerateMipCapable;
-		if (MediaTextureResourceHelpers::SupportsComputeMipGen(InPixelFormat))
-		{
-			OutputCreateFlags |= TexCreate_UAV;
-		}
 
 		// Make sure we only set a number of mips that actually makes sense, given the sample size
 		uint8 MaxMips = FGenericPlatformMath::FloorToInt(FGenericPlatformMath::Log2(static_cast<float>(FGenericPlatformMath::Min(InDim.X, InDim.Y))));
diff --git a/Engine/Source/Runtime/MediaAssets/Private/Misc/MediaTextureResource.h b/Engine/Source/Runtime/MediaAssets/Private/Misc/MediaTextureResource.h
--- a/Engine/Source/Runtime/MediaAssets/Private/Misc/MediaTextureResource.h
+++ b/Engine/Source/Runtime/MediaAssets/Private/Misc/MediaTextureResource.h
@@ -166,7 +166,7 @@ protected:
 	/**
 	 * Create/update output render target as needed
 	 */
-	void CreateOutputRenderTarget(const FIntPoint & InDim, EPixelFormat InPixelFormat, bool bInSRGB, const FLinearColor & InClearColor, uint8 InNumMips);
+	void CreateOutputRenderTarget(const FIntPoint & InDim, EPixelFormat InPixelFormat, bool bInSRGB, const FLinearColor & InClearColor, uint8 InNumMips, bool bNeedsUAVSupport);
 
 	/**
 	 * Caches next available sample from queue in MediaTexture owner to keep single consumer access
