MediaSource: Centralized cache settings so they can be overriden elsewhere.

#rnx

diff --git a/Engine/Plugins/Media/MediaPlate/Source/MediaPlate/Private/MediaPlateComponent.cpp b/Engine/Plugins/Media/MediaPlate/Source/MediaPlate/Private/MediaPlateComponent.cpp
--- a/Engine/Plugins/Media/MediaPlate/Source/MediaPlate/Private/MediaPlateComponent.cpp
+++ b/Engine/Plugins/Media/MediaPlate/Source/MediaPlate/Private/MediaPlateComponent.cpp
@@ -21,6 +21,8 @@ FLazyName UMediaPlateComponent::MediaPlaylistName(TEXT("MediaPlaylist0"));
 UMediaPlateComponent::UMediaPlateComponent(const FObjectInitializer& ObjectInitializer)
 	: Super(ObjectInitializer)
 {
+	CacheSettings.bOverride = true;
+
 	// Set up media component.
 	MediaComponent = CreateDefaultSubobject<UMediaComponent>(MediaComponentName);
 	if (MediaComponent != nullptr)
@@ -152,10 +154,10 @@ bool UMediaPlateComponent::PlayMediaSource(UMediaSource* InMediaSource)
 
 	if (InMediaSource != nullptr)
 	{
+		// Set cache settings.
+		InMediaSource->SetCacheSettings(CacheSettings);
+		
 		// Set media options.
-		InMediaSource->SetMediaOptionBool(TEXT("ImgMediaSmartCacheEnabled"), bSmartCacheEnabled);
-		InMediaSource->SetMediaOptionFloat(TEXT("ImgMediaSmartCacheTimeToLookAhead"), SmartCacheTimeToLookAhead);
-	
 		TObjectPtr<UMediaPlayer> MediaPlayer = GetMediaPlayer();
 		if (MediaPlayer != nullptr)
 		{
diff --git a/Engine/Plugins/Media/MediaPlate/Source/MediaPlate/Public/MediaPlateComponent.h b/Engine/Plugins/Media/MediaPlate/Source/MediaPlate/Public/MediaPlateComponent.h
--- a/Engine/Plugins/Media/MediaPlate/Source/MediaPlate/Public/MediaPlateComponent.h
+++ b/Engine/Plugins/Media/MediaPlate/Source/MediaPlate/Public/MediaPlateComponent.h
@@ -5,6 +5,7 @@
 #include "Components/ActorComponent.h"
 #include "CoreMinimal.h"
 #include "MediaPlayerOptions.h"
+#include "MediaSource.h"
 #include "MediaTextureTracker.h"
 
 #include "MediaPlateComponent.generated.h"
@@ -85,17 +86,9 @@ public:
 	UPROPERTY(EditAnywhere, Category = "MediaPlate")
 	TObjectPtr<UMediaPlaylist> MediaPlaylist;
 
-	/** Enable smart caching for image sequences. */
+	/** Override the default cache settings. */
 	UPROPERTY(EditAnywhere, AdvancedDisplay, Category = "MediaPlate")
-	bool bSmartCacheEnabled = true;
-
-	/**
-	 * The cache will fill up with frames that are up to this time from the current time.
-	 * E.g. if this is 0.2, and we are at time index 5 seconds,
-	 * then we will fill the cache with frames between 5 seconds and 5.2 seconds.
-	 */
-	UPROPERTY(EditAnywhere, AdvancedDisplay, Category = "MediaPlate")
-	float SmartCacheTimeToLookAhead = 0.2f;
+	FMediaSourceCacheSettings CacheSettings;
 
 	/**
 	 * Specify type of mesh used for visible mips and tiles calculations.
diff --git a/Engine/Source/Runtime/MediaAssets/Private/Assets/MediaSource.cpp b/Engine/Source/Runtime/MediaAssets/Private/Assets/MediaSource.cpp
--- a/Engine/Source/Runtime/MediaAssets/Private/Assets/MediaSource.cpp
+++ b/Engine/Source/Runtime/MediaAssets/Private/Assets/MediaSource.cpp
@@ -7,6 +7,12 @@
 #include "Misc/Paths.h"
 #include "StreamMediaSource.h"
 
+void UMediaSource::SetCacheSettings(const FMediaSourceCacheSettings& Settings)
+{
+	SetMediaOptionBool(TEXT("ImgMediaSmartCacheEnabled"), Settings.bOverride);
+	SetMediaOptionFloat(TEXT("ImgMediaSmartCacheTimeToLookAhead"), Settings.TimeToLookAhead);
+}
+
 void UMediaSource::RegisterSpawnFromFileExtension(const FString& Extension,
 	FMediaSourceSpawnDelegate InDelegate)
 {
diff --git a/Engine/Source/Runtime/MediaAssets/Public/MediaSource.h b/Engine/Source/Runtime/MediaAssets/Public/MediaSource.h
--- a/Engine/Source/Runtime/MediaAssets/Public/MediaSource.h
+++ b/Engine/Source/Runtime/MediaAssets/Public/MediaSource.h
@@ -15,6 +15,28 @@
 /** Delegate for creating a media source from a string. */
 DECLARE_DELEGATE_RetVal_TwoParams(UMediaSource*, FMediaSourceSpawnDelegate, const FString&, UObject*);
 
+/** Cache settings to pass to the player. */
+USTRUCT(BlueprintType)
+struct FMediaSourceCacheSettings
+{
+	GENERATED_USTRUCT_BODY()
+
+	/**
+	 * Override the default cache settings.
+	 * Currently only the ImgMedia player supports these settings.
+	 */
+	UPROPERTY(EditAnywhere, Category = "Media Cache")
+	bool bOverride = false;
+
+	/**
+	 * The cache will fill up with frames that are up to this time from the current time.
+	 * E.g. if this is 0.2, and we are at time index 5 seconds,
+	 * then we will fill the cache with frames between 5 seconds and 5.2 seconds.
+	 */
+	UPROPERTY(EditAnywhere, Category = "Media Cache")
+	float TimeToLookAhead = 0.2f;
+};
+
 /**
  * Abstract base class for media sources.
  *
@@ -50,6 +72,11 @@ public:
 	UFUNCTION(BlueprintCallable, Category="Media|MediaSource")
 	virtual bool Validate() const PURE_VIRTUAL(UMediaSource::Validate, return false;);
 
+	/**
+	 * Call this to set cache settings to pass to the player.
+	 */
+	void SetCacheSettings(const FMediaSourceCacheSettings& Settings);
+
 	/**
 	 * Call this to register a callback when someone calls SpawnMediaSourceForString.
 	 * This lets you spawn a media source if the file extension matches what you want.
