Media: Add mip bias support to media texture (resource).

diff --git a/Engine/Source/Runtime/MediaAssets/Private/Assets/MediaTexture.cpp b/Engine/Source/Runtime/MediaAssets/Private/Assets/MediaTexture.cpp
--- a/Engine/Source/Runtime/MediaAssets/Private/Assets/MediaTexture.cpp
+++ b/Engine/Source/Runtime/MediaAssets/Private/Assets/MediaTexture.cpp
@@ -87,6 +87,7 @@ UMediaTexture::UMediaTexture(const FObjectInitializer& ObjectInitializer)
 	, Size(0)
 	, CachedNextSampleTime(FTimespan::MinValue())
 	, TextureNumMips(1)
+	, MipMapBias(0.0f)
 {
 	NeverStream = true;
 	SRGB = true;
@@ -533,3 +534,14 @@ MediaTextureOrientation UMediaTexture::GetCurrentOrientation() const
 {
 	return CurrentOrientation;
 }
+
+float UMediaTexture::GetMipMapBias() const
+{
+	// Clamped to the legal DirectX range.
+	return FMath::Clamp(MipMapBias, -16.0f, 15.99f);
+}
+
+void UMediaTexture::SetMipMapBias(float InMipMapBias)
+{
+	MipMapBias = InMipMapBias;
+}
diff --git a/Engine/Source/Runtime/MediaAssets/Private/Misc/MediaTextureResource.cpp b/Engine/Source/Runtime/MediaAssets/Private/Misc/MediaTextureResource.cpp
--- a/Engine/Source/Runtime/MediaAssets/Private/Misc/MediaTextureResource.cpp
+++ b/Engine/Source/Runtime/MediaAssets/Private/Misc/MediaTextureResource.cpp
@@ -271,6 +271,7 @@ FMediaTextureResource::FMediaTextureResource(UMediaTexture& InOwner, FIntPoint&
 	, bEnableGenMips(InEnableGenMips)
 	, CurrentNumMips(InEnableGenMips ? InNumMips : 1)
 	, CurrentSamplerFilter(ESamplerFilter_Num)
+	, CurrentMipMapBias(-1)
 	, PriorSamples(FPriorSamples::Create())
 {
 #if PLATFORM_ANDROID
@@ -589,17 +590,20 @@ uint32 FMediaTextureResource::GetSizeY() const
 void FMediaTextureResource::SetupSampler()
 {
 	ESamplerFilter OwnerFilter = (CurrentNumMips > 1 || Owner.NewStyleOutput) ? (ESamplerFilter)UDeviceProfileManager::Get().GetActiveProfile()->GetTextureLODSettings()->GetSamplerFilter(&Owner) : SF_Bilinear;
+	float OwnerMipMapBias = Owner.GetMipMapBias();
 
-	if (CurrentSamplerFilter != OwnerFilter)
+	if (CurrentSamplerFilter != OwnerFilter || !FMath::IsNearlyEqual(CurrentMipMapBias, OwnerMipMapBias))
 	{
 		CurrentSamplerFilter = OwnerFilter;
+		CurrentMipMapBias = OwnerMipMapBias;
 
 		// create the sampler state
 		FSamplerStateInitializerRHI SamplerStateInitializer(
 			CurrentSamplerFilter,
 			(Owner.AddressX == TA_Wrap) ? AM_Wrap : ((Owner.AddressX == TA_Clamp) ? AM_Clamp : AM_Mirror),
 			(Owner.AddressY == TA_Wrap) ? AM_Wrap : ((Owner.AddressY == TA_Clamp) ? AM_Clamp : AM_Mirror),
-			AM_Wrap
+			AM_Wrap,
+			CurrentMipMapBias
 		);
 
 		SamplerStateRHI = GetOrCreateSamplerState(SamplerStateInitializer);
diff --git a/Engine/Source/Runtime/MediaAssets/Private/Misc/MediaTextureResource.h b/Engine/Source/Runtime/MediaAssets/Private/Misc/MediaTextureResource.h
--- a/Engine/Source/Runtime/MediaAssets/Private/Misc/MediaTextureResource.h
+++ b/Engine/Source/Runtime/MediaAssets/Private/Misc/MediaTextureResource.h
@@ -222,6 +222,9 @@ private:
 	/** Current texture sampler filter value */
 	ESamplerFilter CurrentSamplerFilter;
 
+	/** Current texture sampler mip bias. */
+	float CurrentMipMapBias;
+
 	/** The current media player facade to get video samples from. */
 	TWeakPtr<FMediaPlayerFacade, ESPMode::ThreadSafe> PlayerFacadePtr;
 
diff --git a/Engine/Source/Runtime/MediaAssets/Public/MediaTexture.h b/Engine/Source/Runtime/MediaAssets/Public/MediaTexture.h
--- a/Engine/Source/Runtime/MediaAssets/Public/MediaTexture.h
+++ b/Engine/Source/Runtime/MediaAssets/Public/MediaTexture.h
@@ -196,6 +196,20 @@ public:
 	 */
 	MediaTextureOrientation GetCurrentOrientation() const;
 
+	/**
+	 * Get the texture's mip-map bias, clamped to a legal range.
+	 * @return Mip-map bias value
+	 */
+	float GetMipMapBias() const;
+
+	/**
+	 * Set texture's mip-map bias, for use by the texture resource sampler.
+	 * Note: UpdateResource() should be called afterwards and the material should be notified.
+	 *
+	 * @param InMipMapBias Bias value.
+	 */
+	void SetMipMapBias(float InMipMapBias);
+
 public:
 
 	//~ UTexture interface.
@@ -297,4 +311,7 @@ private:
 
 	/** Number of mips in the actual output texture */
 	int32 TextureNumMips;
+
+	/** Mip-map bias used by the media texture resource sampler. */
+	float MipMapBias;
 };
