MediaAssets: Replaced UMediaPlayer::GetNumTile with generic GetMediaInfo.
#rnx

diff --git a/Engine/Plugins/Media/ImgMedia/Source/ImgMedia/Private/Player/ImgMediaPlayer.cpp b/Engine/Plugins/Media/ImgMedia/Source/ImgMedia/Private/Player/ImgMediaPlayer.cpp
--- a/Engine/Plugins/Media/ImgMedia/Source/ImgMedia/Private/Player/ImgMediaPlayer.cpp
+++ b/Engine/Plugins/Media/ImgMedia/Source/ImgMedia/Private/Player/ImgMediaPlayer.cpp
@@ -157,6 +157,24 @@ FString FImgMediaPlayer::GetInfo() const
 }
 
 
+FVariant FImgMediaPlayer::GetMediaInfo(FName InfoName) const
+{
+	FVariant Variant;
+
+	// Source num tiles?
+	if (InfoName == UMediaPlayer::MediaInfoNameSourceNumTiles.Resolve())
+	{
+		if (Loader.IsValid())
+		{
+			FIntPoint TileNum(Loader->GetNumTilesX(), Loader->GetNumTilesY());
+			Variant = TileNum;
+		}
+	}
+
+	return Variant;
+}
+
+
 FGuid FImgMediaPlayer::GetPlayerPluginGUID() const
 {
 	static FGuid PlayerPluginGUID(0x0e4a60c0, 0x2c5947ea, 0xb233562a, 0x57e5761c);
@@ -287,18 +305,6 @@ bool FImgMediaPlayer::Open(const TSharedRef<FArchive, ESPMode::ThreadSafe>& /*Ar
 	return false; // not supported
 }
 
-FIntPoint FImgMediaPlayer::GetTileNum() const
-{
-	FIntPoint TileNum = FIntPoint::ZeroValue;
-
-	if (Loader.IsValid())
-	{
-		TileNum.X = Loader->GetNumTilesX();
-		TileNum.Y = Loader->GetNumTilesY();
-	}
-	
-	return TileNum;
-}
 
 void FImgMediaPlayer::TickInput(FTimespan DeltaTime, FTimespan /*Timecode*/)
 {
diff --git a/Engine/Plugins/Media/ImgMedia/Source/ImgMedia/Private/Player/ImgMediaPlayer.h b/Engine/Plugins/Media/ImgMedia/Source/ImgMedia/Private/Player/ImgMediaPlayer.h
--- a/Engine/Plugins/Media/ImgMedia/Source/ImgMedia/Private/Player/ImgMediaPlayer.h
+++ b/Engine/Plugins/Media/ImgMedia/Source/ImgMedia/Private/Player/ImgMediaPlayer.h
@@ -54,6 +54,7 @@ public:
 	virtual IMediaCache& GetCache() override;
 	virtual IMediaControls& GetControls() override;
 	virtual FString GetInfo() const override;
+	virtual FVariant GetMediaInfo(FName InfoName) const override;
 	virtual FGuid GetPlayerPluginGUID() const override;
 	virtual IMediaSamples& GetSamples() override;
 	virtual FString GetStats() const override;
@@ -62,7 +63,6 @@ public:
 	virtual IMediaView& GetView() override;
 	virtual bool Open(const FString& Url, const IMediaOptions* Options) override;
 	virtual bool Open(const TSharedRef<FArchive, ESPMode::ThreadSafe>& Archive, const FString& OriginalUrl, const IMediaOptions* Options) override;
-	virtual FIntPoint GetTileNum() const override;
 	virtual void TickInput(FTimespan DeltaTime, FTimespan Timecode) override;
 	virtual bool FlushOnSeekStarted() const;
 	virtual bool FlushOnSeekCompleted() const;
diff --git a/Engine/Plugins/Media/MediaCompositing/Source/MediaCompositingEditor/Private/Sequencer/MediaThumbnailSection.cpp b/Engine/Plugins/Media/MediaCompositing/Source/MediaCompositingEditor/Private/Sequencer/MediaThumbnailSection.cpp
--- a/Engine/Plugins/Media/MediaCompositing/Source/MediaCompositingEditor/Private/Sequencer/MediaThumbnailSection.cpp
+++ b/Engine/Plugins/Media/MediaCompositing/Source/MediaCompositingEditor/Private/Sequencer/MediaThumbnailSection.cpp
@@ -351,7 +351,8 @@ void FMediaThumbnailSection::DrawMediaInfo(FSequencerSectionPainter& InPainter,
 {
 	// Get tile info.
 	FString TileString;
-	FIntPoint TileNum = MediaPlayer->GetTileNum();
+	FIntPoint TileNum(EForceInit::ForceInitToZero);
+	MediaPlayer->GetMediaInfo<FIntPoint>(TileNum, UMediaPlayer::MediaInfoNameSourceNumTiles.Resolve());
 	int32 TileTotalNum = TileNum.X * TileNum.Y;
 	if (TileTotalNum > 1)
 	{
diff --git a/Engine/Plugins/Media/MediaPlate/Source/MediaPlateEditor/Private/Widgets/SMediaPlateEditorMediaDetails.cpp b/Engine/Plugins/Media/MediaPlate/Source/MediaPlateEditor/Private/Widgets/SMediaPlateEditorMediaDetails.cpp
--- a/Engine/Plugins/Media/MediaPlate/Source/MediaPlateEditor/Private/Widgets/SMediaPlateEditorMediaDetails.cpp
+++ b/Engine/Plugins/Media/MediaPlate/Source/MediaPlateEditor/Private/Widgets/SMediaPlateEditorMediaDetails.cpp
@@ -148,7 +148,8 @@ void SMediaPlateEditorMediaDetails::UpdateDetails()
 		{
 			FrameRate = MediaPlayer->GetVideoTrackFrameRate(INDEX_NONE, INDEX_NONE);
 			Format = MediaPlayer->GetVideoTrackType(INDEX_NONE, INDEX_NONE);
-			FIntPoint NumTiles = MediaPlayer->GetTileNum();
+			FIntPoint NumTiles(EForceInit::ForceInitToZero);
+			MediaPlayer->GetMediaInfo<FIntPoint>(NumTiles, UMediaPlayer::MediaInfoNameSourceNumTiles.Resolve());
 			NumTotalTiles = NumTiles.X * NumTiles.Y;
 		}
 
diff --git a/Engine/Source/Runtime/Media/Public/IMediaPlayer.h b/Engine/Source/Runtime/Media/Public/IMediaPlayer.h
--- a/Engine/Source/Runtime/Media/Public/IMediaPlayer.h
+++ b/Engine/Source/Runtime/Media/Public/IMediaPlayer.h
@@ -9,6 +9,7 @@
 #include "Math/MathFwd.h"
 #include "Misc/Paths.h"
 #include "Misc/Timespan.h"
+#include "Misc/Variant.h"
 
 class FArchive;
 
@@ -181,6 +182,18 @@ public:
 		return Open(Url, Options);
 	}
 
+	/**
+	 * Get information about the media that is playing.
+	 *
+	 * @param	InfoName		Name of the information we want.
+	 * @returns					Requested information, or empty if not available.
+	 * @see						UMediaPlayer::GetMediaInfo.
+	 */
+	virtual FVariant GetMediaInfo(FName InfoName) const
+	{
+		return FVariant();
+	}
+
 	/**
 	 * Get the human readable name of the currently loaded media source.
 	 *
@@ -204,14 +217,6 @@ public:
 		return FText::FromString(FPaths::GetBaseFilename(Url));
 	}
 
-	/**
-	 * Get the number of tiles in each frame.
-	 */
-	virtual FIntPoint GetTileNum() const
-	{
-		return FIntPoint::ZeroValue;
-	}
-
 	/**
 	 * Set the player's globally unique identifier.
 	 *
diff --git a/Engine/Source/Runtime/MediaAssets/Private/Assets/MediaPlayer.cpp b/Engine/Source/Runtime/MediaAssets/Private/Assets/MediaPlayer.cpp
--- a/Engine/Source/Runtime/MediaAssets/Private/Assets/MediaPlayer.cpp
+++ b/Engine/Source/Runtime/MediaAssets/Private/Assets/MediaPlayer.cpp
@@ -30,6 +30,10 @@
 #include "EngineAnalytics.h"
 #endif
 
+
+FLazyName UMediaPlayer::MediaInfoNameSourceNumMips(TEXT("SourceNumMips"));
+FLazyName UMediaPlayer::MediaInfoNameSourceNumTiles(TEXT("SourceNumTiles"));
+
 /* UMediaPlayer structors
  *****************************************************************************/
 
@@ -200,6 +204,12 @@ int32 UMediaPlayer::GetNumTrackFormats(EMediaPlayerTrack TrackType, int32 TrackI
 }
 
 
+FVariant UMediaPlayer::GetMediaInfo(FName InfoName) const
+{
+	return PlayerFacade->GetMediaInfo(InfoName);
+}
+
+
 TSharedRef<FMediaPlayerFacade, ESPMode::ThreadSafe> UMediaPlayer::GetPlayerFacade() const
 {
 	return PlayerFacade.ToSharedRef();
@@ -230,11 +240,6 @@ void UMediaPlayer::GetSupportedRates(TArray<FFloatRange>& OutRates, bool Unthinn
 	Rates.GetRanges((TArray<TRange<float>>&)OutRates);
 }
 
-FIntPoint UMediaPlayer::GetTileNum() const
-{
-	return PlayerFacade->GetTileNum();
-}
-
 FTimespan UMediaPlayer::GetTime() const
 {
 	return PlayerFacade->GetTime();
diff --git a/Engine/Source/Runtime/MediaAssets/Public/MediaPlayer.h b/Engine/Source/Runtime/MediaAssets/Public/MediaPlayer.h
--- a/Engine/Source/Runtime/MediaAssets/Public/MediaPlayer.h
+++ b/Engine/Source/Runtime/MediaAssets/Public/MediaPlayer.h
@@ -13,6 +13,7 @@
 #include "UObject/Object.h"
 #include "UObject/ScriptMacros.h"
 #include "Misc/Guid.h"
+#include "Misc/Variant.h"
 #include "Engine/LatentActionManager.h"
 #include "MediaPlayerOptions.h"
 #include "IMediaTimeSource.h"
@@ -297,14 +298,6 @@ public:
 	UFUNCTION(BlueprintCallable, Category="Media|MediaPlayer")
 	void GetSupportedRates(TArray<FFloatRange>& OutRates, bool Unthinned) const;
 
-	/**
-	 * Get the number of tiles in each frame.
-	 * 
-	 * @return Number of tiles in each dimension.
-	 */
-	UFUNCTION(BlueprintCallable, Category = "Media|MediaPlayer")
-	virtual FIntPoint GetTileNum() const;
-
 	/**
 	 * Get the media's current playback time.
 	 *
@@ -972,6 +965,40 @@ public:
 		return PlayerGuid;
 	}
 
+	/**
+	 * Get information about the media that is playing.
+	 * Not all formats support all information.
+	 * Some possible common information is in this file,
+	 * look for MediaInfoNameSourceNumMips for example.
+	 *
+	 * @param	InfoName		Name of the information we want.
+	 * @returns					Requested information, or empty if not available.
+	 */
+	FVariant GetMediaInfo(FName InfoName) const;
+
+	/**
+	 * Templated version of GetMediaInfo.
+	 * No need to deal with variants.
+	 * 
+	 * @param	T			Type of the information.
+	 * @param	Result		If the information is found and is the correct type, then this will be
+	 *						set to its value. It will not be set otherwise.
+	 * @param	InfoName	Name of the information.
+	 * @returns				True if the information is found, false if not.
+	 */
+	template<typename T>
+	bool GetMediaInfo(T& Result, FName InfoName) const
+	{
+		FVariant Variant = GetMediaInfo(InfoName);
+		if (Variant.GetType() == TVariantTraits<T>::GetType())
+		{
+			Result = Variant.GetValue<T>();
+			return true;
+		}
+
+		return false;
+	}
+
 	/**
 	 * Get the media player facade that manages low-level media players
 	 *
@@ -1045,6 +1072,14 @@ public:
 #endif
 
 public:
+	/**
+	 * These are possible common information for use with GetMediaInfo.
+	 * They also decsribe the type of data that will be needed to retrieve the information.
+	 */
+	 /** Number of mip map levels in the source. Int32.*/
+	static FLazyName MediaInfoNameSourceNumMips;
+	/** Number of tiles (X and Y) in the source. IntPoint.*/
+	static FLazyName MediaInfoNameSourceNumTiles;
 
 	/**
 	 * Duration of samples to cache ahead of the play head.
diff --git a/Engine/Source/Runtime/MediaUtils/Private/MediaPlayerFacade.cpp b/Engine/Source/Runtime/MediaUtils/Private/MediaPlayerFacade.cpp
--- a/Engine/Source/Runtime/MediaUtils/Private/MediaPlayerFacade.cpp
+++ b/Engine/Source/Runtime/MediaUtils/Private/MediaPlayerFacade.cpp
@@ -324,6 +324,17 @@ FString FMediaPlayerFacade::GetInfo() const
 }
 
 
+FVariant FMediaPlayerFacade::GetMediaInfo(FName InfoName) const
+{
+	TSharedPtr<IMediaPlayer, ESPMode::ThreadSafe> CurrentPlayer(Player);
+	if (!CurrentPlayer.IsValid())
+	{
+		return FVariant();
+	}
+	return CurrentPlayer->GetMediaInfo(InfoName);
+}
+
+
 FText FMediaPlayerFacade::GetMediaName() const
 {
 	TSharedPtr<IMediaPlayer, ESPMode::ThreadSafe> CurrentPlayer(Player);
@@ -425,18 +436,6 @@ bool FMediaPlayerFacade::HaveAudioPlayback() const
 	return PrimaryAudioSink.IsValid() && GetSelectedTrack(EMediaTrackType::Audio) != INDEX_NONE;
 }
 
-FIntPoint FMediaPlayerFacade::GetTileNum() const
-{
-	FIntPoint TileNum = FIntPoint::ZeroValue;
-
-	TSharedPtr<IMediaPlayer, ESPMode::ThreadSafe> CurrentPlayer(Player);
-	if (CurrentPlayer.IsValid())
-	{
-		TileNum = CurrentPlayer->GetTileNum();
-	}
-
-	return TileNum;
-}
 
 FTimespan FMediaPlayerFacade::GetTime() const
 {
diff --git a/Engine/Source/Runtime/MediaUtils/Public/MediaPlayerFacade.h b/Engine/Source/Runtime/MediaUtils/Public/MediaPlayerFacade.h
--- a/Engine/Source/Runtime/MediaUtils/Public/MediaPlayerFacade.h
+++ b/Engine/Source/Runtime/MediaUtils/Public/MediaPlayerFacade.h
@@ -24,6 +24,7 @@
 #include "Misc/Guid.h"
 #include "Misc/Optional.h"
 #include "Misc/Timespan.h"
+#include "Misc/Variant.h"
 #include "Templates/Atomic.h"
 #include "Templates/SharedPointer.h"
 #include "UObject/NameTypes.h"
@@ -233,6 +234,15 @@ public:
 	 */
 	FString GetInfo() const;
 
+	/**
+	 * Get information about the media that is playing.
+	 * 
+	 * @param	InfoName		Name of the information we want.
+	 * @returns					Requested information, or empty if not available.
+	 * @see						UMediaPlayer::GetMediaInfo.
+	 */
+	FVariant GetMediaInfo(FName InfoName) const;
+
 	/**
 	 * Get the human readable name of the currently loaded media source.
 	 *
@@ -312,11 +322,6 @@ public:
 	 */
 	TRangeSet<float> GetSupportedRates(bool Unthinned = true) const;
 
-	/**
-	 * Get the number of tiles in each frame.
-	 */
-	virtual FIntPoint GetTileNum() const;
-
 	/**
 	 * Get the media's current playback time.
 	 *
