Refactor of media sound generator to use the new ISoundGenerator interface.

Cleans up the ISoundGenerator in SynthComponent in OnUnregister to make sure it releases the generator data before GC.

diff --git a/Engine/Source/Runtime/AudioMixer/Private/Components/SynthComponent.cpp b/Engine/Source/Runtime/AudioMixer/Private/Components/SynthComponent.cpp
--- a/Engine/Source/Runtime/AudioMixer/Private/Components/SynthComponent.cpp
+++ b/Engine/Source/Runtime/AudioMixer/Private/Components/SynthComponent.cpp
@@ -349,13 +349,17 @@ void USynthComponent::OnUnregister()
 		AudioComponent->DestroyComponent();
 		AudioComponent = nullptr;
 	}
+
+	// Clear out the synth component's reference to the sound generator or it will leak until it gets GC'd
+	// Normally this is ok to wait till GC but some derived synths might need for the handle to be released
+	SoundGenerator.Reset();
 }
 
 void USynthComponent::EndPlay(const EEndPlayReason::Type Reason) 
 {	
 	Super::EndPlay(Reason);
 
-	if (GetOwner() && (Reason == EEndPlayReason::LevelTransition || Reason == EEndPlayReason::RemovedFromWorld))
+	if (GetOwner() && (Reason == EEndPlayReason::LevelTransition || Reason == EEndPlayReason::RemovedFromWorld || Reason == EEndPlayReason::Destroyed))
 	{
 		// If our world or sublevel is going away, stop immediately to prevent the containing world/level from being leaked via hard references from the audio device.
 		Stop();
diff --git a/Engine/Source/Runtime/MediaAssets/Private/Assets/MediaSoundComponent.cpp b/Engine/Source/Runtime/MediaAssets/Private/Assets/MediaSoundComponent.cpp
--- a/Engine/Source/Runtime/MediaAssets/Private/Assets/MediaSoundComponent.cpp
+++ b/Engine/Source/Runtime/MediaAssets/Private/Assets/MediaSoundComponent.cpp
@@ -19,8 +19,6 @@
 #include "MediaPlayer.h"
 #include "MediaPlayerFacade.h"
 
-#include UE_INLINE_GENERATED_CPP_BY_NAME(MediaSoundComponent)
-
 
 DECLARE_FLOAT_COUNTER_STAT(TEXT("MediaUtils MediaSoundComponent Sync"), STAT_MediaUtils_MediaSoundComponentSync, STATGROUP_Media);
 DECLARE_FLOAT_COUNTER_STAT(TEXT("MediaUtils MediaSoundComponent SampleTime"), STAT_MediaUtils_MediaSoundComponentSampleTime, STATGROUP_Media);
@@ -53,6 +51,130 @@ private:
 	TWeakObjectPtr<UMediaSoundComponent> Owner;
 };
 
+FMediaSoundGenerator::FMediaSoundGenerator(FSoundGeneratorParams& InParams)
+	: Params(InParams)
+{
+	// Initialize the settings for the spectrum analyzer
+	SpectrumAnalyzer.Init((float)InParams.SampleRate);
+	Resampler.Initialize(InParams.NumChannels, InParams.SampleRate);
+
+	Audio::FEnvelopeFollowerInitParams EnvelopeInitParams;
+	EnvelopeInitParams.SampleRate = (float)InParams.SampleRate;
+	EnvelopeInitParams.NumChannels = 1; //EnvelopeFollower uses mixed down mono buffer 
+	EnvelopeFollower.Init(EnvelopeInitParams);
+
+	CachedRate = Params.CachedRate;
+	CachedTime = Params.CachedTime;
+	LastPlaySampleTime = Params.LastPlaySampleTime;
+
+	UE_LOG(LogMediaAssets, Verbose, TEXT("MediaSoundComponent: FMediaSoundGenerator Created."));
+}
+
+FMediaSoundGenerator::~FMediaSoundGenerator()
+{
+	UE_LOG(LogMediaAssets, Verbose, TEXT("MediaSoundComponent: FMediaSoundGenerator Destroyed."));
+}
+
+void FMediaSoundGenerator::OnEndGenerate()
+{
+	UE_LOG(LogMediaAssets, Verbose, TEXT("MediaSoundComponent: OnEndGenerate called."));
+	Params.SampleQueue.Reset();
+}
+
+void FMediaSoundGenerator::SetCachedData(float InCachedRate, const FTimespan& InCachedTime)
+{
+	CachedRate = InCachedRate;
+	CachedTime = InCachedTime;
+}
+
+void FMediaSoundGenerator::SetLastPlaySampleTime(const FTimespan& InLastPlaySampleTime)
+{
+	LastPlaySampleTime = InLastPlaySampleTime;
+}
+
+void FMediaSoundGenerator::SetEnableSpectralAnalysis(bool bInSpectralAnlaysisEnabled)
+{
+	FScopeLock Lock(&AnalysisCritSect);
+	Params.bSpectralAnalysisEnabled = bInSpectralAnlaysisEnabled;
+}
+
+void FMediaSoundGenerator::SetEnableEnvelopeFollowing(bool bInEnvelopeFollowingEnabled)
+{
+	FScopeLock Lock(&AnalysisCritSect);
+	Params.bEnvelopeFollowingEnabled = bInEnvelopeFollowingEnabled;
+	CurrentEnvelopeValue = 0.0f;
+}
+
+void FMediaSoundGenerator::SetSpectrumAnalyzerSettings(Audio::FSpectrumAnalyzerSettings::EFFTSize InFFTSize, const TArray<float>& InFrequenciesToAnalyze)
+{
+	FScopeLock Lock(&AnalysisCritSect);
+	Params.SpectrumAnalyzerSettings.FFTSize = InFFTSize;
+	Params.FrequenciesToAnalyze = InFrequenciesToAnalyze;
+	SpectrumAnalyzer.SetSettings(Params.SpectrumAnalyzerSettings);
+}
+
+void FMediaSoundGenerator::SetEnvelopeFollowingSettings(int32 InAttackTimeMsec, int32 InReleaseTimeMsec)
+{
+	FScopeLock Lock(&AnalysisCritSect);
+	Params.EnvelopeFollowerAttackTime = InAttackTimeMsec;
+	Params.EnvelopeFollowerReleaseTime = InReleaseTimeMsec;
+	bEnvelopeFollowerSettingsChanged = true;
+}
+
+void FMediaSoundGenerator::SetSampleQueue(TSharedPtr<FMediaAudioSampleQueue, ESPMode::ThreadSafe>& InSampleQueue)
+{
+	FScopeLock Lock(&SampleQueueCritSect);
+	Params.SampleQueue = InSampleQueue;
+
+	UE_LOG(LogMediaAssets, Verbose, TEXT("MediaSoundComponent: SetSampleQueue called with new sample queue."));
+}
+
+TArray<FMediaSoundComponentSpectralData> FMediaSoundGenerator::GetSpectralData() const
+{
+	FScopeLock Lock(&AnalysisCritSect);
+
+	if (Params.bSpectralAnalysisEnabled)
+	{
+		Audio::FAsyncSpectrumAnalyzerScopeLock AnalyzerBufferLock(&SpectrumAnalyzer);
+
+		TArray<FMediaSoundComponentSpectralData> SpectralData;
+
+		for (float Frequency : Params.FrequenciesToAnalyze)
+		{
+			FMediaSoundComponentSpectralData Data;
+			Data.FrequencyHz = Frequency;
+			Data.Magnitude = SpectrumAnalyzer.GetMagnitudeForFrequency(Frequency);
+			SpectralData.Add(Data);
+		}
+		return SpectralData;
+	}
+	return TArray<FMediaSoundComponentSpectralData>();
+}
+
+TArray<FMediaSoundComponentSpectralData> FMediaSoundGenerator::GetNormalizedSpectralData() const
+{
+	FScopeLock Lock(&AnalysisCritSect);
+
+	if (Params.bSpectralAnalysisEnabled)
+	{
+		Audio::FAsyncSpectrumAnalyzerScopeLock AnalyzerBufferLock(&SpectrumAnalyzer);
+
+		TArray<FMediaSoundComponentSpectralData> SpectralData;
+
+		for (float Frequency : Params.FrequenciesToAnalyze)
+		{
+			FMediaSoundComponentSpectralData Data;
+			Data.FrequencyHz = Frequency;
+			Data.Magnitude = SpectrumAnalyzer.GetNormalizedMagnitudeForFrequency(Frequency);
+			SpectralData.Add(Data);
+		}
+
+		return SpectralData;
+	}
+	return TArray<FMediaSoundComponentSpectralData>();
+}
+
+
 
 static const int32 MaxAudioInputSamples = 8;	// accept at most these many samples into our input queue
 
@@ -69,11 +191,9 @@ UMediaSoundComponent::UMediaSoundComponent(const FObjectInitializer& ObjectIniti
 	, CachedRate(0.0f)
 	, CachedTime(FTimespan::Zero())
 	, RateAdjustment(1.0f)
-	, Resampler(new FMediaAudioResampler)
 	, LastPlaySampleTime(FTimespan::MinValue())
 	, EnvelopeFollowerAttackTime(10)
 	, EnvelopeFollowerReleaseTime(100)
-	, CurrentEnvelopeValue(0.0f)
 	, bSpectralAnalysisEnabled(false)
 	, bEnvelopeFollowingEnabled(false)
 {
@@ -92,8 +212,6 @@ UMediaSoundComponent::UMediaSoundComponent(const FObjectInitializer& ObjectIniti
 
 UMediaSoundComponent::~UMediaSoundComponent()
 {
-	delete Resampler;
-
 	RemoveClockSink();
 }
 
@@ -148,7 +266,7 @@ void UMediaSoundComponent::UpdatePlayer()
 
 		FScopeLock Lock(&CriticalSection);
 		SampleQueue.Reset();
-
+		MediaSoundGenerator.Reset();
 		return;
 	}
 
@@ -156,11 +274,11 @@ void UMediaSoundComponent::UpdatePlayer()
 	TSharedRef<FMediaPlayerFacade, ESPMode::ThreadSafe> PlayerFacade = CurrentPlayerPtr->GetPlayerFacade();
 
 	// We have some audio decoders which are running with a limited amount of pre-allocated audio sample packets. 
-	// When the audio packets are not consumed in the UMediaSoundComponent::OnGenerateAudio method below, these packets are not 
+	// When the audio packets are not consumed in the FMediaSoundGenerator::OnGenerateAudio method below, these packets are not 
 	// returned to the decoder which then cannot produce more audio samples. 
 	//
-	// The UMediaSoundComponent::OnGenerateAudio is only called when our parent USynthComponent it active and
-	// this is conrolled by USynthComponent::Start() and USynthComponent::Stop(). We are tracking a state change here.
+	// The FMediaSoundGenerator::OnGenerateAudio is only called when our parent USynthComponent it active and
+	// this is controlled by USynthComponent::Start() and USynthComponent::Stop(). We are tracking a state change here.
 	if (PlayerFacade != CurrentPlayerFacade)
 	{
 		if (IsActive())
@@ -190,7 +308,20 @@ void UMediaSoundComponent::UpdatePlayer()
 	CachedRate = PlayerFacade->GetRate();
 	CachedTime = PlayerFacade->GetTime();
 
-	PlayerFacade->SetLastAudioRenderedSampleTime(LastPlaySampleTime.Load());
+	if (MediaSoundGenerator.IsValid())
+	{
+		FMediaSoundGenerator* MediaGen = static_cast<FMediaSoundGenerator*>(MediaSoundGenerator.Get());
+
+		// The play time is derived from the media sound generator's OnGenerate callback
+		LastPlaySampleTime = MediaGen->GetLastPlayTime();
+
+		MediaGen->SetCachedData(CachedRate, CachedTime);
+		PlayerFacade->SetLastAudioRenderedSampleTime(LastPlaySampleTime);
+	}
+	else
+	{
+		PlayerFacade->SetLastAudioRenderedSampleTime(FTimespan::MinValue());
+	}
 }
 
 
@@ -240,6 +371,7 @@ void UMediaSoundComponent::OnUnregister()
 		SampleQueue.Reset();
 	}
 	CurrentPlayerFacade.Reset();
+	MediaSoundGenerator.Reset();
 	Super::OnUnregister();
 }
 
@@ -276,6 +408,7 @@ void UMediaSoundComponent::Deactivate()
 			SampleQueue.Reset();
 		}
 		CurrentPlayerFacade.Reset();
+		MediaSoundGenerator.Reset();
 	}
 	Super::Deactivate();
 }
@@ -324,9 +457,6 @@ bool UMediaSoundComponent::Init(int32& SampleRate)
 {
 	Super::Init(SampleRate);
 
-	// Initialize the settings for the spectrum analyzer
-	SpectrumAnalyzer.Init(SampleRate);
-
 	if (Channels == EMediaSoundChannels::Mono)
 	{
 		NumChannels = 1;
@@ -340,51 +470,64 @@ bool UMediaSoundComponent::Init(int32& SampleRate)
 		NumChannels = 8;
 	}
 
-	// increase buffer callback size for media decoding. Media doesn't need fast response time so can decode more per callback.
-	//PreferredBufferLength = NumChannels * 8196;
-
-	Resampler->Initialize(NumChannels, SampleRate);
-
-	Audio::FEnvelopeFollowerInitParams EnvelopeInitParams;
-	EnvelopeInitParams.SampleRate = SampleRate;
-	EnvelopeInitParams.NumChannels = 1; //EnvelopeFollower uses mixed down mono buffer 
-	EnvelopeFollower.Init(EnvelopeInitParams);
-
 	return true;
 }
 
+ISoundGeneratorPtr UMediaSoundComponent::CreateSoundGenerator(const FSoundGeneratorInitParams& InParams)
+{
+	FMediaSoundGenerator::FSoundGeneratorParams Params;
+	Params.SampleRate = InParams.SampleRate;
+	Params.NumChannels = InParams.NumChannels;
+	Params.SampleQueue = SampleQueue;
+
+	Params.bSpectralAnalysisEnabled = bSpectralAnalysisEnabled;
+	Params.bEnvelopeFollowingEnabled = bEnvelopeFollowingEnabled;
+	Params.EnvelopeFollowerAttackTime = EnvelopeFollowerAttackTime;
+	Params.EnvelopeFollowerReleaseTime = EnvelopeFollowerReleaseTime;
+	Params.SpectrumAnalyzerSettings = SpectrumAnalyzerSettings;
+	Params.FrequenciesToAnalyze = FrequenciesToAnalyze;
+	
+	Params.CachedRate = CachedRate;
+	Params.CachedTime = CachedTime;
+	Params.LastPlaySampleTime = LastPlaySampleTime;
 
-int32 UMediaSoundComponent::OnGenerateAudio(float* OutAudio, int32 NumSamples)
+	return MediaSoundGenerator = ISoundGeneratorPtr(new FMediaSoundGenerator(Params));
+}
+ 
+int32 FMediaSoundGenerator::OnGenerateAudio(float* OutAudio, int32 NumSamples)
 {
-	CSV_SCOPED_TIMING_STAT(MediaStreaming, UMediaSoundComponent_OnGenerateAudio);
+	CSV_SCOPED_TIMING_STAT(MediaStreaming, FMediaSoundGenerator_OnGenerateAudio);
 
 	int32 InitialSyncOffset = 0;
+
 	TSharedPtr<FMediaAudioSampleQueue, ESPMode::ThreadSafe> PinnedSampleQueue;
+	// Make sure we don't swap the sample queue ptr while we're generating
 	{
-		FScopeLock Lock(&CriticalSection);
-		PinnedSampleQueue = SampleQueue;
+		FScopeLock Lock(&SampleQueueCritSect);
+		PinnedSampleQueue = Params.SampleQueue;
 	}
 
-	// We have an input queue and are actively playing?
-	if (PinnedSampleQueue.IsValid() && (CachedRate != 0.0f))
+	const float Rate = CachedRate.Load();
+
+	// 	// We have an input queue and are actively playing?
+	if (PinnedSampleQueue.IsValid() && (Rate != 0.0f))
 	{
-		const float Rate = CachedRate.Load();
 		const FTimespan Time = CachedTime.Load();
 
 		{
-			const uint32 FramesRequested = uint32(NumSamples / NumChannels);
+			const uint32 FramesRequested = uint32(NumSamples / Params.NumChannels);
 			uint32 JumpFrame = MAX_uint32;
 			FMediaTimeStamp OutTime = FMediaTimeStamp(FTimespan::Zero());
-			uint32 FramesWritten = Resampler->Generate(OutAudio, OutTime, FramesRequested, Rate, Time, *PinnedSampleQueue, JumpFrame);
+			uint32 FramesWritten = Resampler.Generate(OutAudio, OutTime, FramesRequested, Rate, Time, *PinnedSampleQueue, JumpFrame);
 			if (FramesWritten == 0)
 			{
-				return 0; // no samples available
+				return NumSamples; // no samples available
 			}
 
 			// Fill in any gap left as we didn't have enough data
 			if (FramesWritten < FramesRequested)
 			{
-				memset(OutAudio + FramesWritten * NumChannels, 0, (NumSamples - FramesWritten * NumChannels) * sizeof(float));
+				memset(OutAudio + FramesWritten * Params.NumChannels, 0, (NumSamples - FramesWritten * Params.NumChannels) * sizeof(float));
 			}
 
 			// Update audio time
@@ -397,12 +540,12 @@ int32 UMediaSoundComponent::OnGenerateAudio(float* OutAudio, int32 NumSamples)
 
 		// - - - - - - - - - - - - - - - - - - - - - - - - - - - - - - - - - - - - - - - - - - - - - - - - - - - - - - - - - - - - - - - - - - - -
 
-		if (bSpectralAnalysisEnabled || bEnvelopeFollowingEnabled)
+		if (Params.bSpectralAnalysisEnabled || Params.bEnvelopeFollowingEnabled)
 		{
 			float* BufferToUseForAnalysis = nullptr;
 			int32 NumFrames = NumSamples;
 			
-			if (NumChannels == 2)
+			if (Params.NumChannels == 2)
 			{
 				NumFrames = NumSamples / 2;
 
@@ -411,7 +554,7 @@ int32 UMediaSoundComponent::OnGenerateAudio(float* OutAudio, int32 NumSamples)
 				AudioScratchBuffer.AddUninitialized(NumFrames);
 				BufferToUseForAnalysis = AudioScratchBuffer.GetData();
 				int32 SampleIndex = 0;
-				for (int32 FrameIndex = 0; FrameIndex < NumFrames; ++FrameIndex, SampleIndex += NumChannels)
+				for (int32 FrameIndex = 0; FrameIndex < NumFrames; ++FrameIndex, SampleIndex += Params.NumChannels)
 				{
 					BufferToUseForAnalysis[FrameIndex] = 0.5f * (OutAudio[SampleIndex] + OutAudio[SampleIndex + 1]);
 				}
@@ -421,20 +564,20 @@ int32 UMediaSoundComponent::OnGenerateAudio(float* OutAudio, int32 NumSamples)
 				BufferToUseForAnalysis = OutAudio;
 			}
 
-			if (bSpectralAnalysisEnabled)
+			if (Params.bSpectralAnalysisEnabled)
 			{
 				SpectrumAnalyzer.PushAudio(BufferToUseForAnalysis, NumFrames);
 				SpectrumAnalyzer.PerformAsyncAnalysisIfPossible(true);
 			}
 
 			{
-				FScopeLock ScopeLock(&EnvelopeFollowerCriticalSection);
-				if (bEnvelopeFollowingEnabled)
+				FScopeLock ScopeLock(&AnalysisCritSect);
+				if (Params.bEnvelopeFollowingEnabled)
 				{
 					if (bEnvelopeFollowerSettingsChanged)
 					{
-						EnvelopeFollower.SetAttackTime((float)EnvelopeFollowerAttackTime);
-						EnvelopeFollower.SetReleaseTime((float)EnvelopeFollowerReleaseTime);
+						EnvelopeFollower.SetAttackTime((float)Params.EnvelopeFollowerAttackTime);
+						EnvelopeFollower.SetReleaseTime((float)Params.EnvelopeFollowerReleaseTime);
 
 						bEnvelopeFollowerSettingsChanged = false;
 					}
@@ -456,16 +599,22 @@ int32 UMediaSoundComponent::OnGenerateAudio(float* OutAudio, int32 NumSamples)
 	}
 	else
 	{
-		Resampler->Flush();
+		Resampler.Flush();
 
 		LastPlaySampleTime = FTimespan::MinValue();
 	}
-	return NumSamples;
+ 	return NumSamples;
+
 }
 
 void UMediaSoundComponent::SetEnableSpectralAnalysis(bool bInSpectralAnalysisEnabled)
 {
 	bSpectralAnalysisEnabled = bInSpectralAnalysisEnabled;
+	if (MediaSoundGenerator.IsValid())
+	{
+		FMediaSoundGenerator* MediaGen = static_cast<FMediaSoundGenerator*>(MediaSoundGenerator.Get());
+		MediaGen->SetEnableSpectralAnalysis(bSpectralAnalysisEnabled);
+	}
 }
 
 void UMediaSoundComponent::SetSpectralAnalysisSettings(TArray<float> InFrequenciesToAnalyze, EMediaSoundComponentFFTSize InFFTSize)
@@ -493,29 +642,24 @@ void UMediaSoundComponent::SetSpectralAnalysisSettings(TArray<float> InFrequenci
 	}
 
 	SpectrumAnalyzerSettings.FFTSize = SpectrumAnalyzerSize;
-	SpectrumAnalyzer.SetSettings(SpectrumAnalyzerSettings);
-
 	FrequenciesToAnalyze = InFrequenciesToAnalyze;
+
+	if (MediaSoundGenerator.IsValid())
+	{
+		FMediaSoundGenerator* MediaGen = static_cast<FMediaSoundGenerator*>(MediaSoundGenerator.Get());
+		MediaGen->SetSpectrumAnalyzerSettings(SpectrumAnalyzerSize, InFrequenciesToAnalyze);
+	}
 }
 
 TArray<FMediaSoundComponentSpectralData> UMediaSoundComponent::GetSpectralData()
 {
 	if (bSpectralAnalysisEnabled)
 	{
-		// Locks analyzer output buffer during access to frequency data.
-		Audio::FAsyncSpectrumAnalyzerScopeLock AnalyzerBufferLock(&SpectrumAnalyzer);
-
-		TArray<FMediaSoundComponentSpectralData> SpectralData;
-
-		for (float Frequency : FrequenciesToAnalyze)
+		if (MediaSoundGenerator.IsValid())
 		{
-			FMediaSoundComponentSpectralData Data;
-			Data.FrequencyHz = Frequency;
-			Data.Magnitude = SpectrumAnalyzer.GetMagnitudeForFrequency(Frequency);
-			SpectralData.Add(Data);
+			FMediaSoundGenerator* MediaGen = static_cast<FMediaSoundGenerator*>(MediaSoundGenerator.Get());
+			return MediaGen->GetSpectralData();
 		}
-
-		return SpectralData;
 	}
 	// Empty array if spectrum analysis is not implemented
 	return TArray<FMediaSoundComponentSpectralData>();
@@ -525,20 +669,11 @@ TArray<FMediaSoundComponentSpectralData> UMediaSoundComponent::GetNormalizedSpec
 {
 	if (bSpectralAnalysisEnabled)
 	{
-		// Locks analyzer output buffer during access to frequency data.
-		Audio::FAsyncSpectrumAnalyzerScopeLock AnalyzerBufferLock(&SpectrumAnalyzer);
-
-		TArray<FMediaSoundComponentSpectralData> SpectralData;
-
-		for (float Frequency : FrequenciesToAnalyze)
+		if (MediaSoundGenerator.IsValid())
 		{
-			FMediaSoundComponentSpectralData Data;
-			Data.FrequencyHz = Frequency;
-			Data.Magnitude = SpectrumAnalyzer.GetNormalizedMagnitudeForFrequency(Frequency);
-			SpectralData.Add(Data);
+			FMediaSoundGenerator* MediaGen = static_cast<FMediaSoundGenerator*>(MediaSoundGenerator.Get());
+			return MediaGen->GetNormalizedSpectralData();
 		}
-
-		return SpectralData;
 	}
 	// Empty array if spectrum analysis is not implemented
 	return TArray<FMediaSoundComponentSpectralData>();
@@ -546,22 +681,36 @@ TArray<FMediaSoundComponentSpectralData> UMediaSoundComponent::GetNormalizedSpec
 
 void UMediaSoundComponent::SetEnableEnvelopeFollowing(bool bInEnvelopeFollowing)
 {
-	FScopeLock ScopeLock(&EnvelopeFollowerCriticalSection);
 	bEnvelopeFollowingEnabled = bInEnvelopeFollowing;
-	CurrentEnvelopeValue = 0.0f;
+
+	if (MediaSoundGenerator.IsValid())
+	{
+		FMediaSoundGenerator* MediaGen = static_cast<FMediaSoundGenerator*>(MediaSoundGenerator.Get());
+		MediaGen->SetEnableEnvelopeFollowing(bInEnvelopeFollowing);
+	}
 }
 
 void UMediaSoundComponent::SetEnvelopeFollowingsettings(int32 AttackTimeMsec, int32 ReleaseTimeMsec)
 {
-	FScopeLock ScopeLock(&EnvelopeFollowerCriticalSection);
 	EnvelopeFollowerAttackTime = AttackTimeMsec;
 	EnvelopeFollowerReleaseTime = ReleaseTimeMsec;
-	bEnvelopeFollowerSettingsChanged = true;
+
+	if (MediaSoundGenerator.IsValid())
+	{
+		FMediaSoundGenerator* MediaGen = static_cast<FMediaSoundGenerator*>(MediaSoundGenerator.Get());
+		MediaGen->SetEnvelopeFollowingSettings(EnvelopeFollowerAttackTime, EnvelopeFollowerReleaseTime);
+	}
 }
 
 float UMediaSoundComponent::GetEnvelopeValue() const
 {
-	return CurrentEnvelopeValue;
+	if (MediaSoundGenerator.IsValid())
+	{
+		FMediaSoundGenerator* MediaGen = static_cast<FMediaSoundGenerator*>(MediaSoundGenerator.Get());
+		return MediaGen->GetCurrentEnvelopeValue();
+	}
+
+	return 0.0f;
 }
 
 void UMediaSoundComponent::AddClockSink()
@@ -610,4 +759,3 @@ const FSoundAttenuationSettings* UMediaSoundComponent::GetSelectedAttenuationSet
 
 	return nullptr;
 }
-
diff --git a/Engine/Source/Runtime/MediaAssets/Public/MediaSoundComponent.h b/Engine/Source/Runtime/MediaAssets/Public/MediaSoundComponent.h
--- a/Engine/Source/Runtime/MediaAssets/Public/MediaSoundComponent.h
+++ b/Engine/Source/Runtime/MediaAssets/Public/MediaSoundComponent.h
@@ -16,6 +16,8 @@
 #include "DSP/BufferVectorOperations.h"
 #include "DSP/EnvelopeFollower.h"
 #include "Sound/SoundClass.h"
+#include "Sound/SoundGenerator.h"
+#include "MediaAudioResampler.h"
 
 #include "MediaSoundComponent.generated.h"
 
@@ -66,6 +68,82 @@ struct FMediaSoundComponentSpectralData
 	float Magnitude = 0.0f;
 };
 
+// Class implements an ISoundGenerator to feed decoded audio to audio renderering async tasks
+class FMediaSoundGenerator : public ISoundGenerator
+{
+public:
+	struct FSoundGeneratorParams
+	{
+		int32 SampleRate = 0;
+		int32 NumChannels = 0;
+
+		TSharedPtr<FMediaAudioSampleQueue, ESPMode::ThreadSafe> SampleQueue;
+
+		bool bSpectralAnalysisEnabled = false;
+		bool bEnvelopeFollowingEnabled = false;
+
+		int32 EnvelopeFollowerAttackTime = 0;
+		int32 EnvelopeFollowerReleaseTime = 0;
+
+		Audio::FSpectrumAnalyzerSettings SpectrumAnalyzerSettings;
+		TArray<float> FrequenciesToAnalyze;
+
+		float CachedRate = 0.0f;
+		FTimespan CachedTime;
+		FTimespan LastPlaySampleTime;
+	};
+
+	FMediaSoundGenerator(FSoundGeneratorParams& InParams);
+
+	virtual ~FMediaSoundGenerator();
+
+	virtual void OnEndGenerate() override;
+
+	virtual int32 OnGenerateAudio(float* OutAudio, int32 NumSamples) override;
+
+	void SetCachedData(float InCachedRate, const FTimespan& InCachedTime);
+	void SetLastPlaySampleTime(const FTimespan& InLastPlaySampleTime);
+
+	void SetEnableSpectralAnalysis(bool bInSpectralAnlaysisEnabled);
+	void SetEnableEnvelopeFollowing(bool bInEnvelopeFollowingEnabled);
+
+	void SetSpectrumAnalyzerSettings(Audio::FSpectrumAnalyzerSettings::EFFTSize InFFTSize, const TArray<float>& InFrequenciesToAnalyze);
+	void SetEnvelopeFollowingSettings(int32 InAttackTimeMsec, int32 InReleaseTimeMsec);
+
+	void SetSampleQueue(TSharedPtr<FMediaAudioSampleQueue, ESPMode::ThreadSafe>& InSampleQueue);
+
+	TArray<FMediaSoundComponentSpectralData> GetSpectralData() const;
+	TArray<FMediaSoundComponentSpectralData> GetNormalizedSpectralData() const;
+	float GetCurrentEnvelopeValue() const { return CurrentEnvelopeValue; }
+
+	FTimespan GetLastPlayTime() const { return LastPlaySampleTime.Load(); }
+
+private:
+
+	FSoundGeneratorParams Params;
+
+	/** The audio resampler. */
+	FMediaAudioResampler Resampler;
+
+	/** Scratch buffer to mix in source audio to from decoder */
+	Audio::AlignedFloatBuffer AudioScratchBuffer;
+
+	/** Spectrum analyzer used for analyzing audio in media. */
+	mutable Audio::FAsyncSpectrumAnalyzer SpectrumAnalyzer;
+
+	Audio::FEnvelopeFollower EnvelopeFollower;
+
+	TAtomic<float> CachedRate;
+	TAtomic<FTimespan> CachedTime;
+	TAtomic<FTimespan> LastPlaySampleTime;
+
+	float CurrentEnvelopeValue = 0.0f;
+	bool bEnvelopeFollowerSettingsChanged = false;
+
+	mutable FCriticalSection AnalysisCritSect;
+	mutable FCriticalSection SampleQueueCritSect;
+};
+
 /**
  * Implements a sound component for playing a media player's audio output.
  */
@@ -260,7 +338,8 @@ protected:
 	//~ USynthComponent interface
 
 	virtual bool Init(int32& SampleRate) override;
-	virtual int32 OnGenerateAudio(float* OutAudio, int32 NumSamples) override;
+
+	virtual ISoundGeneratorPtr CreateSoundGenerator(const FSoundGeneratorInitParams& InParams) override;
 
 protected:
 
@@ -278,10 +357,10 @@ protected:
 private:
 
 	/** The player's current play rate (cached for use on audio thread). */
-	TAtomic<float> CachedRate;
+	float CachedRate;
 
 	/** The player's current time (cached for use on audio thread). */
-	TAtomic<FTimespan> CachedTime;
+	FTimespan CachedTime;
 
 	/** Critical section for synchronizing access to PlayerFacadePtr. */
 	FCriticalSection CriticalSection;
@@ -295,33 +374,18 @@ private:
 	/** Adjusts the output sample rate to synchronize audio and media clock. */
 	float RateAdjustment;
 
-	/** The audio resampler. */
-	FMediaAudioResampler* Resampler;
-
 	/** Audio sample queue. */
 	TSharedPtr<FMediaAudioSampleQueue, ESPMode::ThreadSafe> SampleQueue;
 
 	/* Time of last sample played. */
-	TAtomic<FTimespan> LastPlaySampleTime;
+	FTimespan LastPlaySampleTime;
 
 	/** Which frequencies to analyze. */
 	TArray<float> FrequenciesToAnalyze;
-
-	/** The FFT bin-size to use for FFT analysis. Smaller sizes make it more reactive but less accurate in the frequency space. */
-	EMediaSoundComponentFFTSize FFTSize;
-
 	/** Spectrum analyzer used for analyzing audio in media. */
-	Audio::FAsyncSpectrumAnalyzer SpectrumAnalyzer;
 	Audio::FSpectrumAnalyzerSettings SpectrumAnalyzerSettings;
-
-	Audio::FEnvelopeFollower EnvelopeFollower;
 	int32 EnvelopeFollowerAttackTime;
 	int32 EnvelopeFollowerReleaseTime;
-	float CurrentEnvelopeValue;
-	FCriticalSection EnvelopeFollowerCriticalSection;
-
-	/** Scratch buffer to mix in source audio to from decoder */
-	Audio::AlignedFloatBuffer AudioScratchBuffer;
 
 	/** Whether or not spectral analysis is enabled. */
 	bool bSpectralAnalysisEnabled;
@@ -329,9 +393,9 @@ private:
 	/** Whether or not envelope following is enabled. */
 	bool bEnvelopeFollowingEnabled;
 
-	/** Whether or not envelope follower settings changed. */
-	bool bEnvelopeFollowerSettingsChanged;
-
 	/** Holds our clock sink if available. */
 	TSharedPtr<FMediaSoundComponentClockSink, ESPMode::ThreadSafe> ClockSink;
+
+	/** Instance of our media sound generator. This is a non-uobject that is used to feed sink audio to a sound source on the audio render thread (or async task). */
+	ISoundGeneratorPtr MediaSoundGenerator;
 };
