Removing explicit use of RHI's 2D Tetxure references
#rnx

diff --git a/Engine/Source/Runtime/MediaAssets/Private/Misc/MediaTextureResource.cpp b/Engine/Source/Runtime/MediaAssets/Private/Misc/MediaTextureResource.cpp
--- a/Engine/Source/Runtime/MediaAssets/Private/Misc/MediaTextureResource.cpp
+++ b/Engine/Source/Runtime/MediaAssets/Private/Misc/MediaTextureResource.cpp
@@ -409,7 +409,7 @@ void FMediaTextureResource::Render(const FRenderParams& Params)
 					if (Converter->GetConverterInfoFlags() & IMediaTextureSampleConverter::ConverterInfoFlags_PreprocessOnly)
 					{
 						// Preprocess...
-						FTexture2DRHIRef DummyTexture;
+						FTextureRHIRef DummyTexture;
 						if (Converter->Convert(DummyTexture, Hints))
 						{
 							// ...followed by the built in conversion code as needed...
@@ -427,7 +427,7 @@ void FMediaTextureResource::Render(const FRenderParams& Params)
 				else
 				{
 					// The converter will create its own output texture for us to use
-					FTexture2DRHIRef OutTexture;
+					FTextureRHIRef OutTexture;
 					if (Converter->Convert(OutTexture, Hints))
 					{
 						// As the converter created the texture, we might need to convert it even more to make it fit our needs. Check...
@@ -688,7 +688,7 @@ void FMediaTextureResource::ReleaseDynamicRHI()
 }
 
 
- bool FMediaTextureResource::RequiresConversion(const FTexture2DRHIRef& SampleTexture, const FIntPoint & OutputDim, uint8 InNumMips) const
+ bool FMediaTextureResource::RequiresConversion(const FTextureRHIRef& SampleTexture, const FIntPoint & OutputDim, uint8 InNumMips) const
  {
 	 if (Owner.NewStyleOutput)
 	 {
@@ -771,20 +771,19 @@ void FMediaTextureResource::ConvertSample(const TSharedPtr<IMediaTextureSample,
 	const uint8 SampleNumMips = Sample->GetNumMips();
 
 	// get input texture
-	FRHITexture2D* InputTexture = nullptr;
+	FRHITexture* InputTexture = nullptr;
 	{
 		// If the sample already provides a texture resource, we simply use that
 		// as the input texture. If the sample only provides raw data, then we
 		// create our own input render target and copy the data into it.
 
 		FRHITexture* SampleTexture = Sample->GetTexture();
-		FRHITexture2D* SampleTexture2D = (SampleTexture != nullptr) ? SampleTexture->GetTexture2D() : nullptr;
 
-		if (SampleTexture2D)
+		if (SampleTexture)
 		{
 			// Use the sample as source texture...
 
-			InputTexture = SampleTexture2D;
+			InputTexture = SampleTexture;
 			UpdateResourceSize();
 			InputTarget = nullptr;
 		}
@@ -834,7 +833,7 @@ void FMediaTextureResource::ConvertSample(const TSharedPtr<IMediaTextureSample,
 }
 
 
- void FMediaTextureResource::ConvertTextureToOutput(FRHITexture2D* InputTexture, const TSharedPtr<IMediaTextureSample, ESPMode::ThreadSafe>& Sample)
+ void FMediaTextureResource::ConvertTextureToOutput(FRHITexture* InputTexture, const TSharedPtr<IMediaTextureSample, ESPMode::ThreadSafe>& Sample)
  {
 	// perform the conversion
 	FRHICommandListImmediate& CommandList = FRHICommandListExecutor::GetImmediateCommandList();
@@ -1081,19 +1080,18 @@ void FMediaTextureResource::ConvertSample(const TSharedPtr<IMediaTextureSample,
 void FMediaTextureResource::CopySample(const TSharedPtr<IMediaTextureSample, ESPMode::ThreadSafe>& Sample, const FLinearColor& ClearColor, uint8 InNumMips, const FGuid & TextureGUID)
 {
 	FRHITexture* SampleTexture = Sample->GetTexture();
-	FRHITexture2D* SampleTexture2D = (SampleTexture != nullptr) ? SampleTexture->GetTexture2D() : nullptr;
 	const uint8 SampleNumMips = Sample->GetNumMips();
 
 	// If the sample already provides a texture resource, we simply use that
 	// as the output render target. If the sample only provides raw data, then
 	// we create our own output render target and copy the data into it.
 
-	if (SampleTexture2D != nullptr)
+	if (SampleTexture != nullptr)
 	{
 		// Use sample's texture as the new render target - no copy
-		if (TextureRHI != SampleTexture2D)
+		if (TextureRHI != SampleTexture)
 		{
-			UpdateTextureReference(SampleTexture2D);
+			UpdateTextureReference(SampleTexture);
 
 			MipGenerationCache.SafeRelease();
 			OutputTarget.SafeRelease();
@@ -1106,7 +1104,7 @@ void FMediaTextureResource::CopySample(const TSharedPtr<IMediaTextureSample, ESP
 			CreateOutputRenderTarget(Sample->GetOutputDim(), MediaTextureResourceHelpers::GetPixelFormat(Sample), MediaTextureResourceHelpers::RequiresSrgbTexture(Sample), ClearColor, NumMips, bNeedsUAVTexture);
 
 			// Copy data into the output texture to able to add mips later on
-			FRHICommandListExecutor::GetImmediateCommandList().CopyTexture(SampleTexture2D, OutputTarget, FRHICopyTextureInfo());
+			FRHICommandListExecutor::GetImmediateCommandList().CopyTexture(SampleTexture, OutputTarget, FRHICopyTextureInfo());
 		}
 	}
 	else
@@ -1246,7 +1244,7 @@ void FMediaTextureResource::UpdateResourceSize()
 }
 
 
-void FMediaTextureResource::UpdateTextureReference(FRHITexture2D* NewTexture)
+void FMediaTextureResource::UpdateTextureReference(FRHITexture* NewTexture)
 {
 	TextureRHI = NewTexture;
 	RenderTargetTextureRHI = NewTexture;
diff --git a/Engine/Source/Runtime/MediaAssets/Private/Misc/MediaTextureResource.h b/Engine/Source/Runtime/MediaAssets/Private/Misc/MediaTextureResource.h
--- a/Engine/Source/Runtime/MediaAssets/Private/Misc/MediaTextureResource.h
+++ b/Engine/Source/Runtime/MediaAssets/Private/Misc/MediaTextureResource.h
@@ -199,10 +199,10 @@ private:
 	FGuid InitialTextureGuid;
 
 	/** Input render target if the texture samples don't provide one (for conversions). */
-	TRefCountPtr<FRHITexture2D> InputTarget;
+	TRefCountPtr<FRHITexture> InputTarget;
 
 	/** Output render target if the texture samples don't provide one. */
-	TRefCountPtr<FRHITexture2D> OutputTarget;
+	TRefCountPtr<FRHITexture> OutputTarget;
 
 	/** The media texture that owns this resource. */
 	UMediaTexture& Owner;
