Setting the sample queue in the MediaSoundGenerator when it changes.

#rnx

diff --git a/Engine/Source/Runtime/MediaAssets/Private/Assets/MediaSoundComponent.cpp b/Engine/Source/Runtime/MediaAssets/Private/Assets/MediaSoundComponent.cpp
--- a/Engine/Source/Runtime/MediaAssets/Private/Assets/MediaSoundComponent.cpp
+++ b/Engine/Source/Runtime/MediaAssets/Private/Assets/MediaSoundComponent.cpp
@@ -288,6 +288,10 @@ void UMediaSoundComponent::UpdatePlayer()
 			{
 				FScopeLock Lock(&CriticalSection);
 				SampleQueue = NewSampleQueue;
+				if (MediaSoundGenerator.IsValid())
+				{
+					static_cast<FMediaSoundGenerator*>(MediaSoundGenerator.Get())->SetSampleQueue(SampleQueue);
+				}
 			}
 			CurrentPlayerFacade = PlayerFacade;
 		}
