MediaAssets: Fixed media sources being created for non video assets.

diff --git a/Engine/Source/Runtime/MediaAssets/Private/Assets/MediaSource.cpp b/Engine/Source/Runtime/MediaAssets/Private/Assets/MediaSource.cpp
--- a/Engine/Source/Runtime/MediaAssets/Private/Assets/MediaSource.cpp
+++ b/Engine/Source/Runtime/MediaAssets/Private/Assets/MediaSource.cpp
@@ -2,7 +2,6 @@
 
 #include "MediaSource.h"
 
-#include "FileMediaSource.h"
 #include "MediaAssetsPrivate.h"
 #include "Misc/Paths.h"
 #include "StreamMediaSource.h"
@@ -55,13 +54,6 @@ UMediaSource* UMediaSource::SpawnMediaSourceForString(const FString& MediaPath,
 		{
 			MediaSource = Delegate->Execute(MediaPath, Outer);
 		}
-		else
-		{
-			// Try a file media source.
-			TObjectPtr<UFileMediaSource> FileMediaSource = NewObject<UFileMediaSource>(Outer, NAME_None, RF_Transactional);
-			FileMediaSource->SetFilePath(MediaPath);
-			MediaSource = FileMediaSource;
-		}
 	}
 
 	// Validate the media source.
diff --git a/Engine/Source/Runtime/MediaAssets/Private/MediaAssetsModule.cpp b/Engine/Source/Runtime/MediaAssets/Private/MediaAssetsModule.cpp
--- a/Engine/Source/Runtime/MediaAssets/Private/MediaAssetsModule.cpp
+++ b/Engine/Source/Runtime/MediaAssets/Private/MediaAssetsModule.cpp
@@ -3,6 +3,7 @@
 #include "MediaAssetsPrivate.h"
 
 #include "CoreMinimal.h"
+#include "FileMediaSource.h"
 #include "IMediaAssetsModule.h"
 #include "Misc/CoreMisc.h"
 #include "Modules/ModuleManager.h"
@@ -92,8 +93,25 @@ public:
 	
 	//~ IModuleInterface interface
 
-	virtual void StartupModule() override { }
-	virtual void ShutdownModule() override { }
+	virtual void StartupModule() override
+	{
+		// Register media source spawners.
+		FMediaSourceSpawnDelegate SpawnDelegate =
+			FMediaSourceSpawnDelegate::CreateStatic(&FMediaAssetsModule::SpawnMediaSourceForString);
+		for (const FString& Ext : FileExtensions)
+		{
+			UMediaSource::RegisterSpawnFromFileExtension(Ext, SpawnDelegate);
+		}
+	}
+
+	virtual void ShutdownModule() override
+	{
+		// Unregister media source spawners.
+		for (const FString& Ext : FileExtensions)
+		{
+			UMediaSource::UnregisterSpawnFromFileExtension(Ext);
+		}
+	}
 
 	virtual bool SupportsDynamicReloading() override
 	{
@@ -101,9 +119,31 @@ public:
 	}
 
 private:
+	/**
+	 * Creates a media source for MediaPath.
+	 *
+	 * @param	MediaPath		File path to the media.
+	 * @param	Outer			Outer to use for this object.
+	 */
+	static UMediaSource* SpawnMediaSourceForString(const FString& MediaPath, UObject* Outer)
+	{
+		TObjectPtr<UFileMediaSource> MediaSource =
+			NewObject<UFileMediaSource>(Outer, NAME_None, RF_Transactional);
+		MediaSource->SetFilePath(MediaPath);
+
+		return MediaSource;
+	}
+
 	/** All the functions that can get a player from an object. */
 	TArray<FOnGetPlayerFromObject> GetPlayerFromObjectDelegates;
 
+	/** List of file extensions that we support. */
+	const TArray<FString> FileExtensions =
+	{
+		TEXT("mov"),
+		TEXT("mp4"),
+		TEXT("wmv"),
+	};
 };
 
 
