MediaAssets: Added thumbnails for media sources.

diff --git a/Engine/Plugins/Media/ImgMedia/Source/ImgMedia/Private/Assets/ImgMediaSource.cpp b/Engine/Plugins/Media/ImgMedia/Source/ImgMedia/Private/Assets/ImgMediaSource.cpp
--- a/Engine/Plugins/Media/ImgMedia/Source/ImgMedia/Private/Assets/ImgMediaSource.cpp
+++ b/Engine/Plugins/Media/ImgMedia/Source/ImgMedia/Private/Assets/ImgMediaSource.cpp
@@ -221,6 +221,10 @@ void UImgMediaSource::PostEditChangeProperty(FPropertyChangedEvent& PropertyChan
 			GlobalCache->EmptyCache();
 		}
 	}
+	else if (PropertyChangedEvent.GetMemberPropertyName() == GET_MEMBER_NAME_CHECKED(ThisClass, SequencePath))
+	{
+		GenerateThumbnail();
+	}
 
 	Super::PostEditChangeProperty(PropertyChangedEvent);
 }
diff --git a/Engine/Plugins/Media/MediaPlayerEditor/Source/MediaPlayerEditor/Private/Customizations/PlatformMediaSourceCustomization.cpp b/Engine/Plugins/Media/MediaPlayerEditor/Source/MediaPlayerEditor/Private/Customizations/PlatformMediaSourceCustomization.cpp
--- a/Engine/Plugins/Media/MediaPlayerEditor/Source/MediaPlayerEditor/Private/Customizations/PlatformMediaSourceCustomization.cpp
+++ b/Engine/Plugins/Media/MediaPlayerEditor/Source/MediaPlayerEditor/Private/Customizations/PlatformMediaSourceCustomization.cpp
@@ -166,12 +166,17 @@ void FPlatformMediaSourceCustomization::HandleMediaSourceEntryBoxChanged(const F
 
 	for (auto Object : OuterObjects)
 	{
-		TObjectPtr<UMediaSource>& OldMediaSource = Cast<UPlatformMediaSource>(Object)->PlatformMediaSources.FindOrAdd(PlatformName);;
-
-		if (OldMediaSource != AssetData.GetAsset())
+		UPlatformMediaSource* PlatformMediaSource = Cast<UPlatformMediaSource>(Object);
+		if (PlatformMediaSource != nullptr)
 		{
-			Object->Modify(true);
-			OldMediaSource = Cast<UMediaSource>(AssetData.GetAsset());
+			TObjectPtr<UMediaSource>& OldMediaSource = PlatformMediaSource->PlatformMediaSources.FindOrAdd(PlatformName);;
+
+			if (OldMediaSource != AssetData.GetAsset())
+			{
+				Object->Modify(true);
+				OldMediaSource = Cast<UMediaSource>(AssetData.GetAsset());
+				PlatformMediaSource->GenerateThumbnail();
+			}
 		}
 	}
 }
diff --git a/Engine/Plugins/Media/MediaPlayerEditor/Source/MediaPlayerEditor/Private/MediaPlayerEditorModule.cpp b/Engine/Plugins/Media/MediaPlayerEditor/Source/MediaPlayerEditor/Private/MediaPlayerEditorModule.cpp
--- a/Engine/Plugins/Media/MediaPlayerEditor/Source/MediaPlayerEditor/Private/MediaPlayerEditorModule.cpp
+++ b/Engine/Plugins/Media/MediaPlayerEditor/Source/MediaPlayerEditor/Private/MediaPlayerEditorModule.cpp
@@ -38,6 +38,7 @@
 
 #include "Models/MediaPlayerEditorCommands.h"
 #include "Shared/MediaPlayerEditorStyle.h"
+#include "Shared/MediaSourceThumbnailRenderer.h"
 #include "Visualizers/MediaSoundComponentVisualizer.h"
 #include "ToolMenus.h"
 #include "Widgets/SMediaPlayerEditorMedia.h"
@@ -224,6 +225,7 @@ protected:
 	void RegisterThumbnailRenderers()
 	{
 		UThumbnailManager::Get().RegisterCustomRenderer(UMediaTexture::StaticClass(), UTextureThumbnailRenderer::StaticClass());
+		UThumbnailManager::Get().RegisterCustomRenderer(UMediaSource::StaticClass(), UMediaSourceThumbnailRenderer::StaticClass());
 	}
 
 	/** Unregisters all asset thumbnail renderers. */
@@ -231,6 +233,7 @@ protected:
 	{
 		if (UObjectInitialized())
 		{
+			UThumbnailManager::Get().UnregisterCustomRenderer(UMediaSource::StaticClass());
 			UThumbnailManager::Get().UnregisterCustomRenderer(UMediaTexture::StaticClass());
 		}
 	}
diff --git a/Engine/Plugins/Media/MediaPlayerEditor/Source/MediaPlayerEditor/Private/Shared/MediaSourceThumbnailRenderer.cpp b/Engine/Plugins/Media/MediaPlayerEditor/Source/MediaPlayerEditor/Private/Shared/MediaSourceThumbnailRenderer.cpp
new file mode 100644
--- /dev/null
+++ b/Engine/Plugins/Media/MediaPlayerEditor/Source/MediaPlayerEditor/Private/Shared/MediaSourceThumbnailRenderer.cpp
@@ -0,0 +1,48 @@
+// Copyright Epic Games, Inc. All Rights Reserved.
+
+#include "MediaSourceThumbnailRenderer.h"
+#include "CanvasTypes.h"
+#include "MediaSource.h"
+#include "MediaTexture.h"
+
+bool UMediaSourceThumbnailRenderer::CanVisualizeAsset(UObject* Object)
+{
+	// If the object has a new thumbnail then use that.
+	return GetThumbnailTextureFromObject(Object) != nullptr;
+}
+
+void UMediaSourceThumbnailRenderer::GetThumbnailSize(UObject* Object, float Zoom, uint32& OutWidth, uint32& OutHeight) const
+{
+	UMediaTexture* ObjectTexture = GetThumbnailTextureFromObject(Object);
+	if (ObjectTexture != nullptr)
+	{
+		OutWidth = Zoom * ObjectTexture->GetSurfaceWidth();
+		OutHeight = Zoom * ObjectTexture->GetSurfaceHeight();
+	}
+	else
+	{
+		OutWidth = 0;
+		OutHeight = 0;
+	}
+}
+
+void UMediaSourceThumbnailRenderer::Draw(UObject* Object, int32 X, int32 Y, uint32 Width, uint32 Height, FRenderTarget* RenderTarget, FCanvas* Canvas, bool bAdditionalViewFamily)
+{
+	UMediaTexture* ObjectTexture = GetThumbnailTextureFromObject(Object);
+	if (ObjectTexture != nullptr)
+	{
+		Super::Draw(ObjectTexture, X, Y, Width, Height, RenderTarget, Canvas, bAdditionalViewFamily);
+	}
+}
+
+UMediaTexture* UMediaSourceThumbnailRenderer::GetThumbnailTextureFromObject(UObject* Object) const
+{
+	UMediaSource* MediaSource = Cast<UMediaSource>(Object);
+	if (MediaSource != nullptr)
+	{
+		UMediaTexture* Texture = MediaSource->GetThumbnail();
+		return Texture;
+	}
+	
+	return nullptr;
+}
diff --git a/Engine/Plugins/Media/MediaPlayerEditor/Source/MediaPlayerEditor/Private/Shared/MediaSourceThumbnailRenderer.h b/Engine/Plugins/Media/MediaPlayerEditor/Source/MediaPlayerEditor/Private/Shared/MediaSourceThumbnailRenderer.h
new file mode 100644
--- /dev/null
+++ b/Engine/Plugins/Media/MediaPlayerEditor/Source/MediaPlayerEditor/Private/Shared/MediaSourceThumbnailRenderer.h
@@ -0,0 +1,33 @@
+// Copyright Epic Games, Inc. All Rights Reserved.
+
+#pragma once
+
+#include "ThumbnailRendering/TextureThumbnailRenderer.h"
+
+#include "MediaSourceThumbnailRenderer.generated.h"
+
+class UMediaTexture;
+
+/**
+ * Renders thumbnails for media sources.
+ */
+UCLASS()
+class UMediaSourceThumbnailRenderer : public UTextureThumbnailRenderer
+{
+	GENERATED_BODY()
+
+public:
+
+	// UThumbnailRenderer interface.
+	virtual bool CanVisualizeAsset(UObject* Object) override;
+	virtual void GetThumbnailSize(UObject* Object, float Zoom, uint32& OutWidth, uint32& OutHeight) const override;
+	virtual void Draw(UObject* Object, int32 X, int32 Y, uint32 Width, uint32 Height, FRenderTarget*, FCanvas* Canvas, bool bAdditionalViewFamily) override;
+
+protected:
+
+	/**
+	 * Gets the thumnbail from the object if the object is a media source.
+	 */
+	virtual UMediaTexture* GetThumbnailTextureFromObject(UObject* Object) const;
+
+};
diff --git a/Engine/Source/Runtime/MediaAssets/MediaAssets.Build.cs b/Engine/Source/Runtime/MediaAssets/MediaAssets.Build.cs
--- a/Engine/Source/Runtime/MediaAssets/MediaAssets.Build.cs
+++ b/Engine/Source/Runtime/MediaAssets/MediaAssets.Build.cs
@@ -26,6 +26,10 @@ namespace UnrealBuildTool.Rules
 
 			if (Target.bBuildEditor)
 			{
+				PrivateDependencyModuleNames.AddRange(
+					new string[] {
+						"UnrealEd",
+					});
 				PrivateIncludePathModuleNames.Add("TargetPlatform");
 			}
 		}
diff --git a/Engine/Source/Runtime/MediaAssets/Private/Assets/FileMediaSource.cpp b/Engine/Source/Runtime/MediaAssets/Private/Assets/FileMediaSource.cpp
--- a/Engine/Source/Runtime/MediaAssets/Private/Assets/FileMediaSource.cpp
+++ b/Engine/Source/Runtime/MediaAssets/Private/Assets/FileMediaSource.cpp
@@ -54,6 +54,7 @@ void UFileMediaSource::PostEditChangeProperty(FPropertyChangedEvent& PropertyCha
 	if (PropertyChangedEvent.GetPropertyName() == GET_MEMBER_NAME_CHECKED(UFileMediaSource, FilePath))
 	{
 		ClearResolvedFullPath();
+		GenerateThumbnail();
 	}
 }
 #endif
diff --git a/Engine/Source/Runtime/MediaAssets/Private/Assets/MediaSource.cpp b/Engine/Source/Runtime/MediaAssets/Private/Assets/MediaSource.cpp
--- a/Engine/Source/Runtime/MediaAssets/Private/Assets/MediaSource.cpp
+++ b/Engine/Source/Runtime/MediaAssets/Private/Assets/MediaSource.cpp
@@ -6,6 +6,11 @@
 #include "Misc/Paths.h"
 #include "StreamMediaSource.h"
 
+#if WITH_EDITOR
+#include "Misc/MediaSourceRenderer.h"
+#include "UObject/Package.h"
+#endif
+
 #include UE_INLINE_GENERATED_CPP_BY_NAME(MediaSource)
 
 void UMediaSource::SetCacheSettings(const FMediaSourceCacheSettings& Settings)
@@ -14,6 +19,20 @@ void UMediaSource::SetCacheSettings(const FMediaSourceCacheSettings& Settings)
 	SetMediaOptionFloat(TEXT("ImgMediaSmartCacheTimeToLookAhead"), Settings.TimeToLookAhead);
 }
 
+#if WITH_EDITOR
+
+void UMediaSource::GenerateThumbnail()
+{
+	if (MediaSourceRenderer == nullptr)
+	{
+		MediaSourceRenderer = NewObject<UMediaSourceRenderer>(GetTransientPackage());
+	}
+
+	ThumbnailImage = MediaSourceRenderer->Open(this);
+}
+
+#endif // WITH_EDITOR
+
 void UMediaSource::RegisterSpawnFromFileExtension(const FString& Extension,
 	FMediaSourceSpawnDelegate InDelegate)
 {
@@ -75,6 +94,15 @@ TMap<FString, FMediaSourceSpawnDelegate>& UMediaSource::GetSpawnFromFileExtensio
 	return Delegates;
 }
 
+void UMediaSource::BeginDestroy()
+{
+#if WITH_EDITOR
+	MediaSourceRenderer = nullptr;
+#endif // WITH_EDITOR
+
+	Super::BeginDestroy();
+}
+
 /* IMediaOptions interface
  *****************************************************************************/
 
diff --git a/Engine/Source/Runtime/MediaAssets/Private/Assets/StreamMediaSource.cpp b/Engine/Source/Runtime/MediaAssets/Private/Assets/StreamMediaSource.cpp
--- a/Engine/Source/Runtime/MediaAssets/Private/Assets/StreamMediaSource.cpp
+++ b/Engine/Source/Runtime/MediaAssets/Private/Assets/StreamMediaSource.cpp
@@ -19,3 +19,16 @@ bool UStreamMediaSource::Validate() const
 	return StreamUrl.Contains(TEXT("://"));
 }
 
+#if WITH_EDITOR
+
+void UStreamMediaSource::PostEditChangeProperty(FPropertyChangedEvent& PropertyChangedEvent)
+{
+	Super::PostEditChangeProperty(PropertyChangedEvent);
+
+	if (PropertyChangedEvent.GetPropertyName() == GET_MEMBER_NAME_CHECKED(UStreamMediaSource, StreamUrl))
+	{
+		GenerateThumbnail();
+	}
+}
+
+#endif
diff --git a/Engine/Source/Runtime/MediaAssets/Private/Misc/MediaSourceRenderer.cpp b/Engine/Source/Runtime/MediaAssets/Private/Misc/MediaSourceRenderer.cpp
new file mode 100644
--- /dev/null
+++ b/Engine/Source/Runtime/MediaAssets/Private/Misc/MediaSourceRenderer.cpp
@@ -0,0 +1,123 @@
+
+// Copyright Epic Games, Inc. All Rights Reserved.
+
+#include "Misc/MediaSourceRenderer.h"
+
+#include "MediaPlayer.h"
+#include "MediaSource.h"
+#include "MediaTexture.h"
+#include "UObject/Package.h"
+
+#if WITH_EDITOR
+
+UMediaTexture* UMediaSourceRenderer::Open(UMediaSource* InMediaSource)
+{
+	if ((InMediaSource != nullptr) && (InMediaSource->Validate()))
+	{
+		// Set up the player.
+		if (MediaPlayer == nullptr)
+		{
+			MediaPlayer = NewObject<UMediaPlayer>(GetTransientPackage());
+			MediaPlayer->OnSeekCompleted.AddDynamic(this, &UMediaSourceRenderer::OnSeekCompleted);
+		}
+		else
+		{
+			MediaPlayer->Close();
+		}
+
+		// Set up the texture.
+		if (MediaTexture == nullptr)
+		{
+			MediaTexture = NewObject<UMediaTexture>(GetTransientPackage());
+			MediaTexture->NewStyleOutput = true;
+		}
+		MediaTexture->CurrentAspectRatio = 0.0f;
+		MediaTexture->SetMediaPlayer(MediaPlayer.Get());
+		MediaTexture->UpdateResource();
+		MediaSource = InMediaSource;
+
+		// Start playing the media.
+		bIsSeekActive = false;
+		FMediaPlayerOptions Options;
+		Options.PlayOnOpen = EMediaPlayerOptionBooleanOverride::Enabled;
+		Options.Loop = EMediaPlayerOptionBooleanOverride::Disabled;
+		bool bIsPlaying = MediaPlayer->OpenSourceWithOptions(InMediaSource, Options);
+		if (bIsPlaying == false)
+		{
+			Close();
+		}
+	}
+
+	return MediaTexture.Get();
+}
+
+void UMediaSourceRenderer::Tick(float DeltaTime)
+{
+	if (MediaPlayer != nullptr)
+	{
+		// Is the texture ready?
+		// The aspect ratio will change when we have something.
+		if (MediaTexture->CurrentAspectRatio != 0.0f)
+		{
+			// Send this event so the content browser can update our thumbnail.
+			if (MediaSource != nullptr)
+			{
+				FPropertyChangedEvent EmptyPropertyChangedEvent(nullptr);
+				FCoreUObjectDelegates::OnObjectPropertyChanged.Broadcast(MediaSource,
+					EmptyPropertyChangedEvent);
+			}
+			Close();
+		}
+		else
+		{
+			if (MediaPlayer->GetUrl().IsEmpty())
+			{
+				Close();
+			}
+			else if (MediaPlayer->IsPreparing() == false)
+			{
+				// Don't advance playback.
+				if (MediaPlayer->IsPlaying())
+				{
+					MediaPlayer->SetRate(0.0f);
+				}
+
+				// If we can seek, then seek.
+				if ((MediaPlayer->IsBuffering() == false) && (MediaPlayer->IsReady()))
+				{
+					if (bIsSeekActive == false)
+					{
+						bIsSeekActive = true;
+						const FTimespan MediaDuration = MediaPlayer->GetDuration();
+						FTimespan SeekTime = MediaDuration * 0.3f;
+						MediaPlayer->Seek(SeekTime);
+					}
+				}
+			}
+		}
+	}
+}
+
+#endif // WITH_EDITOR
+
+void UMediaSourceRenderer::OnSeekCompleted()
+{
+#if WITH_EDITOR
+	bIsSeekActive = false;
+#endif // WITH_EDITOR
+}
+
+void UMediaSourceRenderer::Close()
+{
+#if WITH_EDITOR
+	if (MediaTexture != nullptr)
+	{
+		MediaTexture->SetMediaPlayer(nullptr);
+	}
+	if (MediaPlayer != nullptr)
+	{
+		MediaPlayer->Close();
+		MediaPlayer = nullptr;
+	}
+#endif // WITH_EDITOR
+}
diff --git a/Engine/Source/Runtime/MediaAssets/Private/Misc/MediaSourceRenderer.h b/Engine/Source/Runtime/MediaAssets/Private/Misc/MediaSourceRenderer.h
new file mode 100644
--- /dev/null
+++ b/Engine/Source/Runtime/MediaAssets/Private/Misc/MediaSourceRenderer.h
@@ -0,0 +1,70 @@
+// Copyright Epic Games, Inc. All Rights Reserved.
+
+#pragma once
+
+#include "UObject/Object.h"
+
+#if WITH_EDITOR
+#include "TickableEditorObject.h"
+#endif // WITH_EDITOR
+
+#include "MediaSourceRenderer.generated.h"
+
+class UMediaPlayer;
+class UMediaSource;
+class UMediaTexture;
+
+/** Renders a media source to a texture in editor builds. */
+UCLASS()
+class UMediaSourceRenderer : public UObject
+#if WITH_EDITOR
+	, public FTickableEditorObject
+#endif // WITH_EDITOR
+{
+	GENERATED_BODY()
+
+public:
+
+#if WITH_EDITOR
+	/**
+	 * Open the media source to render a texture for.
+	 * 
+	 * @param	InMediaSource		Media source to play.
+	 * @return	Media texture that will hold the image.
+	 */
+	UMediaTexture* Open(UMediaSource* InMediaSource);
+
+	/** FTickableEditorObject interface */
+	virtual void Tick(float DeltaTime) override;
+	virtual TStatId GetStatId() const override { RETURN_QUICK_DECLARE_CYCLE_STAT(UMediaSourceRenderer, STATGROUP_Tickables); }
+#endif // WITH_EDITOR
+
+private:
+	/**
+	 * Callback when the media player is done seeking.
+	 */
+	UFUNCTION()
+	void OnSeekCompleted();
+
+	/**
+	 * Cleans everything up.
+	 * The media textue will remain so it can be used/reused.
+	 */
+	void Close();
+
+	/** Holds the player we are using. */
+	UPROPERTY(Transient)
+	TObjectPtr<UMediaPlayer> MediaPlayer = nullptr;
+
+	/** Holds the media source we are using. */
+	UPROPERTY(Transient)
+	TObjectPtr<UMediaSource> MediaSource = nullptr;
+
+	/** Holds the media texture we are using. */
+	UPROPERTY(Transient)
+	TObjectPtr<UMediaTexture> MediaTexture = nullptr;
+
+	/** True if we are currently seeking. */
+	bool bIsSeekActive = false;
+	
+};
diff --git a/Engine/Source/Runtime/MediaAssets/Public/MediaSource.h b/Engine/Source/Runtime/MediaAssets/Public/MediaSource.h
--- a/Engine/Source/Runtime/MediaAssets/Public/MediaSource.h
+++ b/Engine/Source/Runtime/MediaAssets/Public/MediaSource.h
@@ -22,6 +22,8 @@
 
 class FVariant;
 class UMediaSource;
+class UMediaSourceRenderer;
+class UMediaTexture;
 struct FFrame;
 
 /** Delegate for creating a media source from a string. */
@@ -89,6 +91,21 @@ public:
 	 */
 	void SetCacheSettings(const FMediaSourceCacheSettings& Settings);
 
+
+#if WITH_EDITOR
+
+	/**
+	 * Starts the process to generate a thumbnail.
+	 */
+	void GenerateThumbnail();
+
+	/**
+	 * Gets our thumbnail texture, if any.
+	 */
+	UMediaTexture* GetThumbnail() const { return ThumbnailImage; }
+
+#endif // WITH_EDITOR
+
 	/**
 	 * Call this to register a callback when someone calls SpawnMediaSourceForString.
 	 * This lets you spawn a media source if the file extension matches what you want.
@@ -116,6 +133,8 @@ public:
 	static UMediaSource* SpawnMediaSourceForString(const FString& MediaPath, UObject* Outer);
 
 public:
+	//~ UObject interface
+	virtual void BeginDestroy() override;
 
 	//~ IMediaOptions interface
 
@@ -162,4 +181,16 @@ private:
 	 * Get a mapping of file extensions to spawn delegates.
 	 */
 	static TMap<FString, FMediaSourceSpawnDelegate>& GetSpawnFromFileExtensionDelegates();
+
+#if WITH_EDITORONLY_DATA
+
+	/** The thumbnail image.*/
+	UPROPERTY(Transient)
+	TObjectPtr<UMediaTexture> ThumbnailImage = nullptr;
+
+	/** Renders thumnbnails for us. */
+	UPROPERTY(Transient)
+	TObjectPtr<UMediaSourceRenderer> MediaSourceRenderer = nullptr;
+
+#endif
 };
diff --git a/Engine/Source/Runtime/MediaAssets/Public/StreamMediaSource.h b/Engine/Source/Runtime/MediaAssets/Public/StreamMediaSource.h
--- a/Engine/Source/Runtime/MediaAssets/Public/StreamMediaSource.h
+++ b/Engine/Source/Runtime/MediaAssets/Public/StreamMediaSource.h
@@ -30,4 +30,9 @@ public:
 
 	virtual FString GetUrl() const override;
 	virtual bool Validate() const override;
+
+#if WITH_EDITOR
+	virtual void PostEditChangeProperty(FPropertyChangedEvent& PropertyChangedEvent) override;
+#endif
+
 };
