Changing HDR output normalization from 80nits == 1.0 to 100nits when playing back video in MediaFramework
#rnx

diff --git a/Engine/Source/Runtime/Media/Public/IMediaTextureSample.h b/Engine/Source/Runtime/Media/Public/IMediaTextureSample.h
--- a/Engine/Source/Runtime/Media/Public/IMediaTextureSample.h
+++ b/Engine/Source/Runtime/Media/Public/IMediaTextureSample.h
@@ -124,6 +124,7 @@ enum class EMediaOrientation
 	CW270
 };
 
+static constexpr float kMediaSample_HDR_NitsNormalizationFactor = 1.0f / 100.0f;
 
 /**
  * Interface for media texture samples.
diff --git a/Engine/Source/Runtime/MediaAssets/Private/Misc/MediaTextureResource.cpp b/Engine/Source/Runtime/MediaAssets/Private/Misc/MediaTextureResource.cpp
--- a/Engine/Source/Runtime/MediaAssets/Private/Misc/MediaTextureResource.cpp
+++ b/Engine/Source/Runtime/MediaAssets/Private/Misc/MediaTextureResource.cpp
@@ -1021,11 +1021,8 @@ void FMediaTextureResource::ConvertSample(const TSharedPtr<IMediaTextureSample,
 					// Setup conversion from Rec2020 to current working color space
 					const UE::Color::FColorSpace& Working = UE::Color::FColorSpace::GetWorking();
 					FMatrix44f ColorSpaceMtx = UE::Color::Transpose<float>(Working.GetXYZToRgb()) * Sample->GetGamutToXYZMatrix();
-					if (1) //Working.IsSRGB()) // vvv seems UE does this for all colorspaces
-					{
-						// Add scale so we get a proper 1.0 at 80nits
-						ColorSpaceMtx = ColorSpaceMtx.ApplyScale(1.0f / 80.0f);
-					}
+					// Normalize output (e.g. 80 or 100 nits == 1.0)
+					ColorSpaceMtx = ColorSpaceMtx.ApplyScale(kMediaSample_HDR_NitsNormalizationFactor);
 
 					FIntPoint TexDim = InputTexture->GetSizeXY();
 					if (InputTexture->GetFormat() == PF_P010)
