Media Framework: Add a just in time option to MediaTexture.

This allows for a just in time render of the media when when MediaTexture creator knows how late in the rendering pipeline it is needed, which can reduce latency for streamed media by up to one frame.

diff --git a/Engine/Source/Runtime/MediaAssets/Private/Assets/MediaTexture.cpp b/Engine/Source/Runtime/MediaAssets/Private/Assets/MediaTexture.cpp
--- a/Engine/Source/Runtime/MediaAssets/Private/Assets/MediaTexture.cpp
+++ b/Engine/Source/Runtime/MediaAssets/Private/Assets/MediaTexture.cpp
@@ -454,10 +454,26 @@ void UMediaTexture::TickResource(FTimespan Timecode)
 	
 	// redraw texture resource on render thread
 	FMediaTextureResource* ResourceParam = (FMediaTextureResource*)GetResource();
+
+	const ERenderMode RenderModeParam = GetRenderMode();
+
 	ENQUEUE_RENDER_COMMAND(MediaTextureResourceRender)(
-		[ResourceParam, RenderParams](FRHICommandListImmediate& RHICmdList)
+		[ResourceParam, RenderParams, RenderModeParam](FRHICommandListImmediate& RHICmdList)
 		{
-			ResourceParam->Render(RenderParams);
+			check(ResourceParam);
+
+			if (RenderModeParam == ERenderMode::JustInTime)
+			{
+				// Cache the render params if this is a just in time render mode. 
+				// User must call JustInTimeRender to update the resource.
+				ResourceParam->SetJustInTimeRenderParams(RenderParams);
+			}
+			else
+			{
+				// Otherwise, render the texture right away
+				ResourceParam->ResetJustInTimeRenderParams();
+				ResourceParam->Render(RenderParams);
+			}
 		});
 
 	
@@ -465,6 +481,14 @@ void UMediaTexture::TickResource(FTimespan Timecode)
 	bIsCleared = ((AutoClear) && (bIsSampleValid == false));
 }
 
+void UMediaTexture::JustInTimeRender()
+{
+	if (FMediaTextureResource* MediaResource = static_cast<FMediaTextureResource*>(GetResource()))
+	{
+		MediaResource->JustInTimeRender();
+	}
+}
+
 void UMediaTexture::UpdateSampleInfo(const TSharedPtr<IMediaTextureSample, ESPMode::ThreadSafe> & Sample)
 {
 	CurrentAspectRatio = (float)Sample->GetAspectRatio();
diff --git a/Engine/Source/Runtime/MediaAssets/Private/Misc/MediaTextureResource.cpp b/Engine/Source/Runtime/MediaAssets/Private/Misc/MediaTextureResource.cpp
--- a/Engine/Source/Runtime/MediaAssets/Private/Misc/MediaTextureResource.cpp
+++ b/Engine/Source/Runtime/MediaAssets/Private/Misc/MediaTextureResource.cpp
@@ -731,7 +731,33 @@ void FMediaTextureResource::ReleaseDynamicRHI()
 		 (Format != PF_FloatRGB) &&
 		 (Format != PF_FloatRGBA));
  }
+ 
 
+void FMediaTextureResource::SetJustInTimeRenderParams(const FRenderParams& InJustInTimeRenderParams)
+{
+	check(IsInRenderingThread());
+
+	JustInTimeRenderParams = MakeUnique<FRenderParams>();
+	*JustInTimeRenderParams = InJustInTimeRenderParams;
+}
+
+void FMediaTextureResource::ResetJustInTimeRenderParams()
+{
+	check(IsInRenderingThread());
+
+	JustInTimeRenderParams.Reset();
+}
+
+void FMediaTextureResource::JustInTimeRender()
+{
+	check(IsInRenderingThread());
+
+	if (JustInTimeRenderParams.IsValid())
+	{
+		Render(*JustInTimeRenderParams);
+		JustInTimeRenderParams.Reset();
+	}
+}
 
  bool FMediaTextureResource::RequiresConversion(const TSharedPtr<IMediaTextureSample, ESPMode::ThreadSafe>& Sample, uint8 InNumMips) const
  {
diff --git a/Engine/Source/Runtime/MediaAssets/Private/Misc/MediaTextureResource.h b/Engine/Source/Runtime/MediaAssets/Private/Misc/MediaTextureResource.h
--- a/Engine/Source/Runtime/MediaAssets/Private/Misc/MediaTextureResource.h
+++ b/Engine/Source/Runtime/MediaAssets/Private/Misc/MediaTextureResource.h
@@ -104,6 +104,15 @@ public:
 	 */
 	void FlushPendingData();
 
+	/** Sets the just in time render parameters for later use when JustInTimeRender() gets called */
+	void SetJustInTimeRenderParams(const FRenderParams& InJustInTimeRenderParams);
+
+	/** Clears the just in time render params, in which case calling JustInTimeRender() would have no effect */
+	void ResetJustInTimeRenderParams();
+
+	/** Render the texture using the cached FRenderParams. Call from render thread only. */
+	void JustInTimeRender();
+
 public:
 
 	//~ FRenderTarget interface
@@ -236,4 +245,8 @@ private:
 
 	/** cached params etc. for use with mip generator */
 	TRefCountPtr<IPooledRenderTarget> MipGenerationCache;
+
+	/** Cached FRenderParams, used when JustInTimeRender() gets called. */
+	TUniquePtr<FRenderParams> JustInTimeRenderParams;
+
 };
diff --git a/Engine/Source/Runtime/MediaAssets/Public/MediaTexture.h b/Engine/Source/Runtime/MediaAssets/Public/MediaTexture.h
--- a/Engine/Source/Runtime/MediaAssets/Public/MediaTexture.h
+++ b/Engine/Source/Runtime/MediaAssets/Public/MediaTexture.h
@@ -21,7 +21,6 @@ class FMediaTextureClockSink;
 class IMediaTextureSample;
 class UMediaPlayer;
 
-
 UENUM()
 enum MediaTextureOutputFormat : int
 {
@@ -51,6 +50,17 @@ class MEDIAASSETS_API UMediaTexture
 {
 	GENERATED_UCLASS_BODY()
 
+public:
+
+	/** Possible render modes of this media texture. */
+	enum class ERenderMode
+	{
+		Default = 0,
+		JustInTime, // Will defer rendering this media texture until its consumer calls its JustInTimeRender function.
+	};
+
+public:
+
 	/** The addressing mode to use for the X axis. */
 	UPROPERTY(EditAnywhere, BlueprintReadWrite, Category="MediaTexture", meta=(DisplayName="X-axis Tiling Method"), AssetRegistrySearchable, AdvancedDisplay)
 	TEnumAsByte<TextureAddress> AddressX;
@@ -210,6 +220,21 @@ public:
 	 */
 	void SetMipMapBias(float InMipMapBias);
 
+	/** Renders this media texture. Only has an effect if its RenderMode is ERenderMode::JustInTime */
+	virtual void JustInTimeRender();
+
+	/** Sets the ERenderMode of this media texture */
+	void SetRenderMode(ERenderMode InRenderMode)
+	{
+		RenderMode = InRenderMode;
+	}
+
+	/** Returns the ERenderMode of this media texture */
+	ERenderMode GetRenderMode()
+	{
+		return RenderMode;
+	}
+
 public:
 
 	//~ UTexture interface.
@@ -314,4 +339,7 @@ private:
 
 	/** Mip-map bias used by the media texture resource sampler. */
 	float MipMapBias;
+
+	/** Current render mode of this media texture. Can be changed using SetRenderMode() */
+	ERenderMode RenderMode = ERenderMode::Default;
 };
