MediaCompositing: Fixed dead media player ticks accumulating when scrubbing in/out of sections.

#rnx

diff --git a/Engine/Plugins/Media/MediaCompositing/Source/MediaCompositing/Private/MovieScene/MovieSceneMediaData.cpp b/Engine/Plugins/Media/MediaCompositing/Source/MediaCompositing/Private/MovieScene/MovieSceneMediaData.cpp
--- a/Engine/Plugins/Media/MediaCompositing/Source/MediaCompositing/Private/MovieScene/MovieSceneMediaData.cpp
+++ b/Engine/Plugins/Media/MediaCompositing/Source/MediaCompositing/Private/MovieScene/MovieSceneMediaData.cpp
@@ -27,6 +27,7 @@ FMovieSceneMediaData::~FMovieSceneMediaData()
 	{
 		MediaPlayer->OnMediaEvent().RemoveAll(this);
 		MediaPlayer->Close();
+		MediaPlayer->CleanUpBeforeDestroy();
 		MediaPlayer->RemoveFromRoot();
 	}
 }
diff --git a/Engine/Source/Runtime/MediaAssets/Private/Assets/MediaPlayer.cpp b/Engine/Source/Runtime/MediaAssets/Private/Assets/MediaPlayer.cpp
--- a/Engine/Source/Runtime/MediaAssets/Private/Assets/MediaPlayer.cpp
+++ b/Engine/Source/Runtime/MediaAssets/Private/Assets/MediaPlayer.cpp
@@ -143,6 +143,23 @@ void UMediaPlayer::Close()
 	PlayOnNext = false;
 }
 
+static const FName MediaModuleName("Media");
+void UMediaPlayer::CleanUpBeforeDestroy()
+{
+	if (!HasAnyFlags(RF_ClassDefaultObject))
+	{
+		IMediaModule* MediaModule = FModuleManager::LoadModulePtr<IMediaModule>(MediaModuleName);
+
+		if (MediaModule != nullptr)
+		{
+			UnregisterWithMediaModule();
+			MediaModule->GetTicker().RemoveTickable(PlayerFacade.ToSharedRef());
+		}
+
+		PlayerFacade->Close();
+	}
+}
+
 
 int32 UMediaPlayer::GetAudioTrackChannels(int32 TrackIndex, int32 FormatIndex) const
 {
@@ -800,21 +817,9 @@ void UMediaPlayer::ResumePIE()
 /* UObject overrides
  *****************************************************************************/
 
-static const FName MediaModuleName("Media");
 void UMediaPlayer::BeginDestroy()
 {
-	if (!HasAnyFlags(RF_ClassDefaultObject))
-	{
-		IMediaModule* MediaModule = FModuleManager::LoadModulePtr<IMediaModule>(MediaModuleName);
-
-		if (MediaModule != nullptr)
-		{
-			UnregisterWithMediaModule();
-			MediaModule->GetTicker().RemoveTickable(PlayerFacade.ToSharedRef());
-		}
-
-		PlayerFacade->Close();
-	}
+	CleanUpBeforeDestroy();
 
 	Super::BeginDestroy();
 }
diff --git a/Engine/Source/Runtime/MediaAssets/Public/MediaPlayer.h b/Engine/Source/Runtime/MediaAssets/Public/MediaPlayer.h
--- a/Engine/Source/Runtime/MediaAssets/Public/MediaPlayer.h
+++ b/Engine/Source/Runtime/MediaAssets/Public/MediaPlayer.h
@@ -962,6 +962,13 @@ public:
 	FOnMediaPlayerMediaEvent OnMetadataChanged;
 public:
 
+	/**
+	 * Clean up before this object is destroyed. Normally you would not need to call this,
+	 * but if you want to clean up resources before garbage collection (e.g. you are in the editor)
+	 * then this could be helpful.
+	 */
+	void CleanUpBeforeDestroy();
+
 	/**
 	 * Get the Guid associated with this media player
 	 *
