Fixing possible crash due to race condtion in MediaTextureResource

diff --git a/Engine/Source/Runtime/MediaAssets/Private/Misc/MediaTextureResource.cpp b/Engine/Source/Runtime/MediaAssets/Private/Misc/MediaTextureResource.cpp
--- a/Engine/Source/Runtime/MediaAssets/Private/Misc/MediaTextureResource.cpp
+++ b/Engine/Source/Runtime/MediaAssets/Private/Misc/MediaTextureResource.cpp
@@ -302,6 +302,7 @@ FMediaTextureResource::FMediaTextureResource(UMediaTexture& InOwner, FIntPoint&
 
 void FMediaTextureResource::FlushPendingData()
 {
+	FScopeLock Lock(&PriorSamplesCS);
 	PriorSamples = FPriorSamples::Create();
 }
 
@@ -317,7 +318,13 @@ void FMediaTextureResource::Render(const FRenderParams& Params)
 	SCOPE_CYCLE_COUNTER(STAT_MediaAssets_MediaTextureResourceRender);
 	CSV_SCOPED_TIMING_STAT(MediaStreaming, FMediaTextureResource_Render);
 
-	PriorSamples->Update();
+	TSharedPtr<FPriorSamples, ESPMode::ThreadSafe> LocalPriorSamples;
+	{
+	FScopeLock Lock(&PriorSamplesCS);
+	LocalPriorSamples = PriorSamples;
+	}
+
+	LocalPriorSamples->Update();
 
 	FLinearColor Rotation(1, 0, 0, 1);
 	FLinearColor Offset(0, 0, 0, 0);
@@ -482,7 +489,7 @@ void FMediaTextureResource::Render(const FRenderParams& Params)
 			if (CurrentSample)
 			{
 				// If we had a current sample (directly used as output), we can now schedule its retirement
-				PriorSamples->Retire(CurrentSample);
+				LocalPriorSamples->Retire(CurrentSample);
 				CurrentSample = nullptr;
 			}
 
@@ -490,7 +497,7 @@ void FMediaTextureResource::Render(const FRenderParams& Params)
 			if (OutputTarget == RenderTargetTextureRHI)
 			{
 				// Yes, we can schedule the actual sample for retirement right away
-				PriorSamples->Retire(Sample);
+				LocalPriorSamples->Retire(Sample);
 			}
 			else
 			{
@@ -540,7 +547,7 @@ void FMediaTextureResource::Render(const FRenderParams& Params)
 			if (CurrentSample)
 			{
 				// If we had a current sample (directly used as output), we can now schedule its retirement
-				PriorSamples->Retire(CurrentSample);
+				LocalPriorSamples->Retire(CurrentSample);
 				CurrentSample = nullptr;
 			}
 		}
diff --git a/Engine/Source/Runtime/MediaAssets/Private/Misc/MediaTextureResource.h b/Engine/Source/Runtime/MediaAssets/Private/Misc/MediaTextureResource.h
--- a/Engine/Source/Runtime/MediaAssets/Private/Misc/MediaTextureResource.h
+++ b/Engine/Source/Runtime/MediaAssets/Private/Misc/MediaTextureResource.h
@@ -242,6 +242,8 @@ private:
 
 	/** prior samples not yet ready for retirement as GPU may still actively use them */
 	TSharedRef<FPriorSamples, ESPMode::ThreadSafe> PriorSamples;
+	/** prior samples CS */
+	FCriticalSection PriorSamplesCS;
 
 	/** cached params etc. for use with mip generator */
 	TRefCountPtr<IPooledRenderTarget> MipGenerationCache;
