MediaTextureResource conversions for CharUYVY are now handled commonly with other YUV 422 formats
#rnx

diff --git a/Engine/Shaders/Private/MediaShaders.usf b/Engine/Shaders/Private/MediaShaders.usf
--- a/Engine/Shaders/Private/MediaShaders.usf
+++ b/Engine/Shaders/Private/MediaShaders.usf
@@ -249,6 +249,7 @@ void YUVv216ConvertPS(
 	// Assemble YUV data from input
 	InRaw = (YUVv216ConvertUB.IsARGBFmt != 0u) ? InRaw.zyxw : InRaw;								// PF_ARGB8 vs. PF_ABGR16 ordering
 	float4 CbY0CrY1 = (YUVv216ConvertUB.IsCbY0CrY1 != 0u) ? InRaw : InRaw.yxwz;						// CbY0CrY1 vs. Y0CbY1Cr ordering
+	CbY0CrY1 = (YUVv216ConvertUB.SwapChroma != 0u) ? CbY0CrY1.zyxw : CbY0CrY1;						// Cb/Cr swap
 	float3 YUV = (uint(InUV.x * YUVv216ConvertUB.OutputDimX) & 1u) ? CbY0CrY1.wxz : CbY0CrY1.yxz;	// Decode to YCbCr for odd/even pixels
 
 	// YUV -> RGB (using a plain 4x4 matrix - not the common 3x3 & offset setup)
diff --git a/Engine/Source/Runtime/MediaAssets/Private/Misc/MediaTextureResource.cpp b/Engine/Source/Runtime/MediaAssets/Private/Misc/MediaTextureResource.cpp
--- a/Engine/Source/Runtime/MediaAssets/Private/Misc/MediaTextureResource.cpp
+++ b/Engine/Source/Runtime/MediaAssets/Private/Misc/MediaTextureResource.cpp
@@ -1030,15 +1030,6 @@ void FMediaTextureResource::ConvertSample(const TSharedPtr<IMediaTextureSample,
 				}
 				break;
 
-				case EMediaTextureSampleFormat::CharUYVY:
-				{
-					TShaderMapRef<FUYVYConvertPS> ConvertShader(ShaderMap);
-					GraphicsPSOInit.BoundShaderState.PixelShaderRHI = ConvertShader.GetPixelShader();
-					SetGraphicsPipelineState(CommandList, GraphicsPSOInit, 0);
-					SetShaderParametersLegacyPS(CommandList, ConvertShader, InputTexture, YUVToRGBMatrix, YUVOffset, bIsSampleOutputSrgb);
-				}
-				break;
-
 				case EMediaTextureSampleFormat::CharYVYU:
 				{
 					TShaderMapRef<FYVYUConvertPS> ConvertShader(ShaderMap);
@@ -1051,6 +1042,7 @@ void FMediaTextureResource::ConvertSample(const TSharedPtr<IMediaTextureSample,
 				case EMediaTextureSampleFormat::CharYUY2:		// Y0CbY1Cr
 				case EMediaTextureSampleFormat::Char2VUY:		// CbY0CrY1
 				case EMediaTextureSampleFormat::YUVv216:		// CbY0CrY1
+				case EMediaTextureSampleFormat::CharUYVY:		// CrY0CbY1
 				{
 					// Get the complete matrix to convert sample data to RGB
 					auto YUVMtx = Sample->GetSampleToRGBMatrix();
@@ -1064,7 +1056,11 @@ void FMediaTextureResource::ConvertSample(const TSharedPtr<IMediaTextureSample,
 					TShaderMapRef<FYUVv216ConvertPS> ConvertShader(ShaderMap);
 					GraphicsPSOInit.BoundShaderState.PixelShaderRHI = ConvertShader.GetPixelShader();
 					SetGraphicsPipelineState(CommandList, GraphicsPSOInit, 0);
-					SetShaderParametersLegacyPS(CommandList, ConvertShader, InputTexture, OutputDim, YUVMtx, Sample->GetEncodingType() == UE::Color::EEncoding::sRGB, Sample->GetEncodingType() == UE::Color::EEncoding::ST2084, ColorSpaceMtx, Sample->GetFormat() != EMediaTextureSampleFormat::CharYUY2, InputTexture->GetFormat() == PF_A8R8G8B8);
+					SetShaderParametersLegacyPS(CommandList, ConvertShader, InputTexture, OutputDim, YUVMtx, Sample->GetEncodingType() == UE::Color::EEncoding::sRGB, Sample->GetEncodingType() == UE::Color::EEncoding::ST2084, ColorSpaceMtx,
+												Sample->GetFormat() != EMediaTextureSampleFormat::CharYUY2,		// Y or Cb first
+												InputTexture->GetFormat() == PF_A8R8G8B8,						// ARGB vs. ABGR
+												Sample->GetFormat() == EMediaTextureSampleFormat::CharUYVY		// Cb / Cr swap
+												);
 				}
 				break;
 
diff --git a/Engine/Source/Runtime/RenderCore/Private/MediaShaders.cpp b/Engine/Source/Runtime/RenderCore/Private/MediaShaders.cpp
--- a/Engine/Source/Runtime/RenderCore/Private/MediaShaders.cpp
+++ b/Engine/Source/Runtime/RenderCore/Private/MediaShaders.cpp
@@ -586,6 +586,7 @@ SHADER_PARAMETER(FMatrix44f, CSTransform)
 SHADER_PARAMETER(uint32, EOTF)
 SHADER_PARAMETER(uint32, IsCbY0CrY1)
 SHADER_PARAMETER(uint32, IsARGBFmt)
+SHADER_PARAMETER(uint32, SwapChroma)
 SHADER_PARAMETER(float, OutputDimX)
 SHADER_PARAMETER(float, OutputDimY)
 SHADER_PARAMETER_TEXTURE(Texture2D<float4>, YUVTexture)
@@ -595,7 +596,7 @@ IMPLEMENT_GLOBAL_SHADER_PARAMETER_STRUCT(FYUVv216ConvertUB, "YUVv216ConvertUB");
 IMPLEMENT_SHADER_TYPE(, FYUVv216ConvertPS, TEXT("/Engine/Private/MediaShaders.usf"), TEXT("YUVv216ConvertPS"), SF_Pixel);
 
 
-void FYUVv216ConvertPS::SetParameters(FRHIBatchedShaderParameters& BatchedParameters, TRefCountPtr<FRHITexture2D> YUVTexture, const FIntPoint& OutputDimensions, const FMatrix44f& ColorTransform, bool bSrgbToLinear, bool bIsST2084, const FMatrix44f& CSTransform, bool bIsCbY0CrY1, bool bIsARGBFmt)
+void FYUVv216ConvertPS::SetParameters(FRHIBatchedShaderParameters& BatchedParameters, TRefCountPtr<FRHITexture2D> YUVTexture, const FIntPoint& OutputDimensions, const FMatrix44f& ColorTransform, bool bSrgbToLinear, bool bIsST2084, const FMatrix44f& CSTransform, bool bIsCbY0CrY1, bool bIsARGBFmt, bool bSwapChroma)
 {
 	FYUVv216ConvertUB UB;
 	{
@@ -604,6 +605,7 @@ void FYUVv216ConvertPS::SetParameters(FRHIBatchedShaderParameters& BatchedParame
 		UB.EOTF = bSrgbToLinear ? 1 : (bIsST2084 ? 2 : 0);
 		UB.IsCbY0CrY1 = bIsCbY0CrY1;
 		UB.IsARGBFmt = bIsARGBFmt;
+		UB.SwapChroma = bSwapChroma;
 		UB.OutputDimX = (float)OutputDimensions.X;
 		UB.OutputDimY = (float)OutputDimensions.Y;
 		UB.YUVTexture = YUVTexture;
diff --git a/Engine/Source/Runtime/RenderCore/Public/MediaShaders.h b/Engine/Source/Runtime/RenderCore/Public/MediaShaders.h
--- a/Engine/Source/Runtime/RenderCore/Public/MediaShaders.h
+++ b/Engine/Source/Runtime/RenderCore/Public/MediaShaders.h
@@ -475,7 +475,7 @@ public:
 		: FGlobalShader(Initializer)
 	{ }
 
-	RENDERCORE_API void SetParameters(FRHIBatchedShaderParameters& BatchedParameters, TRefCountPtr<FRHITexture2D> YUVTexture, const FIntPoint& OutputDimensions, const FMatrix44f& ColorTransform, bool SrgbToLinear, bool bIsST2084, const FMatrix44f& CSTransform, bool bIsCbY0CrY1, bool bIsARGBFmt);
+	RENDERCORE_API void SetParameters(FRHIBatchedShaderParameters& BatchedParameters, TRefCountPtr<FRHITexture2D> YUVTexture, const FIntPoint& OutputDimensions, const FMatrix44f& ColorTransform, bool SrgbToLinear, bool bIsST2084, const FMatrix44f& CSTransform, bool bIsCbY0CrY1, bool bIsARGBFmt, bool bSwapChroma);
 };
 
 
