Fixing up Cb/Cr order for 2VUY format in MediaTextureResource conversion code
#rnx

diff --git a/Engine/Source/Runtime/MediaAssets/Private/Misc/MediaTextureResource.cpp b/Engine/Source/Runtime/MediaAssets/Private/Misc/MediaTextureResource.cpp
--- a/Engine/Source/Runtime/MediaAssets/Private/Misc/MediaTextureResource.cpp
+++ b/Engine/Source/Runtime/MediaAssets/Private/Misc/MediaTextureResource.cpp
@@ -1042,7 +1042,7 @@ void FMediaTextureResource::ConvertSample(const TSharedPtr<IMediaTextureSample,
 				case EMediaTextureSampleFormat::CharYUY2:		// Y0CbY1Cr
 				case EMediaTextureSampleFormat::Char2VUY:		// CbY0CrY1
 				case EMediaTextureSampleFormat::YUVv216:		// CbY0CrY1
-				case EMediaTextureSampleFormat::CharUYVY:		// CrY0CbY1
+				case EMediaTextureSampleFormat::CharUYVY:		// CbY0CrY1
 				{
 					// Get the complete matrix to convert sample data to RGB
 					auto YUVMtx = Sample->GetSampleToRGBMatrix();
@@ -1057,9 +1057,9 @@ void FMediaTextureResource::ConvertSample(const TSharedPtr<IMediaTextureSample,
 					GraphicsPSOInit.BoundShaderState.PixelShaderRHI = ConvertShader.GetPixelShader();
 					SetGraphicsPipelineState(CommandList, GraphicsPSOInit, 0);
 					SetShaderParametersLegacyPS(CommandList, ConvertShader, InputTexture, OutputDim, YUVMtx, Sample->GetEncodingType() == UE::Color::EEncoding::sRGB, Sample->GetEncodingType() == UE::Color::EEncoding::ST2084, ColorSpaceMtx,
-												Sample->GetFormat() != EMediaTextureSampleFormat::CharYUY2,		// Y or Cb first
-												InputTexture->GetFormat() == PF_A8R8G8B8,						// ARGB vs. ABGR
-												Sample->GetFormat() == EMediaTextureSampleFormat::CharUYVY		// Cb / Cr swap
+												Sample->GetFormat() != EMediaTextureSampleFormat::CharYUY2,	// Y or Cb first
+												InputTexture->GetFormat() == PF_B8G8R8A8,					// ARGB vs. ABGR (memory order)
+												false														// No Cb / Cr swap
 												);
 				}
 				break;
