Some more cleanup and unification for MFW texturen sample conversion code & improving handling of 10-bit data
#rnx

diff --git a/Engine/Plugins/Experimental/AVCodecs/AVCodecsCore/Source/AVCodecsCoreRHI/Private/Video/Resources/VideoResourceRHI.cpp b/Engine/Plugins/Experimental/AVCodecs/AVCodecsCore/Source/AVCodecsCoreRHI/Private/Video/Resources/VideoResourceRHI.cpp
--- a/Engine/Plugins/Experimental/AVCodecs/AVCodecsCore/Source/AVCodecsCoreRHI/Private/Video/Resources/VideoResourceRHI.cpp
+++ b/Engine/Plugins/Experimental/AVCodecs/AVCodecsCore/Source/AVCodecsCoreRHI/Private/Video/Resources/VideoResourceRHI.cpp
@@ -306,12 +306,16 @@ TSharedPtr<FVideoResourceRHI> FVideoResourceRHI::TransformResource(FVideoDescrip
 						GraphicsPSOInit.BoundShaderState.VertexShaderRHI = VertexShader.GetVertexShader();
 
 						// TODO (aidan.possemiers) do we need get the not scaled version of this matrix
-						FMatrix YUVToRGBMatrix = MediaShaders::YuvToRgbRec709Scaled;
+						FMatrix PreMtx = FMatrix::Identity;
+						PreMtx.M[0][3] = -MediaShaders::YUVOffset8bits.X;
+						PreMtx.M[1][3] = -MediaShaders::YUVOffset8bits.Y;
+						PreMtx.M[2][3] = -MediaShaders::YUVOffset8bits.Z;
+						auto YUVToRGBMatrix = FMatrix44f(MediaShaders::YuvToRgbRec709Scaled * PreMtx);
 
 						TShaderMapRef<FNV12ConvertAsBytesPS> ConvertShader(ShaderMap);
 						GraphicsPSOInit.BoundShaderState.PixelShaderRHI = ConvertShader.GetPixelShader();
 						SetGraphicsPipelineState(RHICmdList, GraphicsPSOInit, 0);
-						SetShaderParametersLegacyPS(RHICmdList, ConvertShader, Source, OutputDim, YUVToRGBMatrix, MediaShaders::YUVOffset8bits, true);
+						SetShaderParametersLegacyPS(RHICmdList, ConvertShader, Source, OutputDim, YUVToRGBMatrix, true, FMatrix44f::Identity, false);
 
 						// draw full size quad into render target
 						FBufferRHIRef VertexBuffer = CreateTempMediaVertexBuffer();
@@ -372,12 +376,16 @@ void FVideoResourceRHI::TransformResourceTo(FRHICommandListImmediate& RHICmdList
 					GraphicsPSOInit.BoundShaderState.VertexShaderRHI = VertexShader.GetVertexShader();
 
 					// TODO (aidan.possemiers) do we need get the not scaled version of this matrix
-					FMatrix YUVToRGBMatrix = MediaShaders::YuvToRgbRec709Scaled;
+					FMatrix PreMtx = FMatrix::Identity;
+					PreMtx.M[0][3] = -MediaShaders::YUVOffset8bits.X;
+					PreMtx.M[1][3] = -MediaShaders::YUVOffset8bits.Y;
+					PreMtx.M[2][3] = -MediaShaders::YUVOffset8bits.Z;
+					auto YUVToRGBMatrix = FMatrix44f(MediaShaders::YuvToRgbRec709Scaled * PreMtx);
 
 					TShaderMapRef<FNV12ConvertAsBytesPS> ConvertShader(ShaderMap);
 					GraphicsPSOInit.BoundShaderState.PixelShaderRHI = ConvertShader.GetPixelShader();
 					SetGraphicsPipelineState(RHICmdList, GraphicsPSOInit, 0);
-					SetShaderParametersLegacyPS(RHICmdList, ConvertShader, Raw.Texture, OutputDim, YUVToRGBMatrix, MediaShaders::YUVOffset8bits, true);
+					SetShaderParametersLegacyPS(RHICmdList, ConvertShader, Raw.Texture, OutputDim, YUVToRGBMatrix, true, FMatrix44f::Identity, false);
 
 					// draw full size quad into render target
 					FBufferRHIRef VertexBuffer = CreateTempMediaVertexBuffer();
diff --git a/Engine/Plugins/Media/ElectraUtil/Source/ElectraSamples/Private/ElectraTextureSampleBase.cpp b/Engine/Plugins/Media/ElectraUtil/Source/ElectraSamples/Private/ElectraTextureSampleBase.cpp
--- a/Engine/Plugins/Media/ElectraUtil/Source/ElectraSamples/Private/ElectraTextureSampleBase.cpp
+++ b/Engine/Plugins/Media/ElectraUtil/Source/ElectraSamples/Private/ElectraTextureSampleBase.cpp
@@ -28,6 +28,10 @@ void IElectraTextureSampleBase::Initialize(FVideoDecoderOutput* InVideoDecoderOu
 		{
 			NumBits = 8;
 		}
+		else if (PixFmt == PF_A2B10G10R10)
+		{
+			NumBits = 10;
+		}
 		else if (PixFmt == PF_P010)
 		{
 			NumBits = 16;
@@ -40,8 +44,14 @@ void IElectraTextureSampleBase::Initialize(FVideoDecoderOutput* InVideoDecoderOu
 
 	// Prepare YUV -> RGB matrix containing all necessary offsets and scales to produce RGB straight from sample data
 	const FMatrix* Mtx = bFullRange ? &MediaShaders::YuvToRgbRec709Unscaled : &MediaShaders::YuvToRgbRec709Scaled;
-	FVector Off = (NumBits == 8) ? (bFullRange ? MediaShaders::YUVOffsetNoScale8bits : MediaShaders::YUVOffset8bits)
-								 : (bFullRange ? MediaShaders::YUVOffsetNoScale16bits : MediaShaders::YUVOffset16bits);
+	FVector Off;
+	switch (NumBits)
+	{
+		case 8:		Off = bFullRange ? MediaShaders::YUVOffsetNoScale8bits  : MediaShaders::YUVOffset8bits; break;
+		case 10:	Off = bFullRange ? MediaShaders::YUVOffsetNoScale10bits : MediaShaders::YUVOffset10bits; break;
+		case 32:	Off = bFullRange ? MediaShaders::YUVOffsetNoScaleFloat  : MediaShaders::YUVOffsetFloat; break;
+		default:	Off = bFullRange ? MediaShaders::YUVOffsetNoScale16bits : MediaShaders::YUVOffset16bits; break;
+	}
 	// Correctional scale for input data
 	// (data should be placed in the upper 10-bits of the 16-bit texture channels, but some platforms do not do this)
 	float DataScale = GetSampleDataScale(false);
diff --git a/Engine/Shaders/Private/MediaShaders.usf b/Engine/Shaders/Private/MediaShaders.usf
--- a/Engine/Shaders/Private/MediaShaders.usf
+++ b/Engine/Shaders/Private/MediaShaders.usf
@@ -53,6 +53,39 @@ void BMPConvertPS(
 }
 
 
+float3 ApplyColorProcessingRGB(float3 Color, float4x4 CSMtx, uint EOTF)
+{
+	// Go linear
+	if (EOTF != 0)
+	{
+		Color = (EOTF == 1) ? sRGBToLinear(Color) : ST2084ToLinear(Color);
+	}
+
+	// Adjust colorspace
+	float4 Col4 = float4(Color, 1);
+	Color.x = dot(CSMtx[0], Col4);
+	Color.y = dot(CSMtx[1], Col4);
+	Color.z = dot(CSMtx[2], Col4);
+
+	return Color;
+}
+
+
+float3 ApplyColorProcessing(float3 YUV, float4x4 YUVMtx, float4x4 CSMtx, uint EOTF)
+{
+	float3 Color;
+	float4 YUV4 = float4(YUV, 1);
+	Color.x = dot(YUVMtx[0], YUV4);
+	Color.y = dot(YUVMtx[1], YUV4);
+	Color.z = dot(YUVMtx[2], YUV4);
+
+	// Make sure we safely stay above zero
+	Color = max(Color, 0.0f);
+
+	return ApplyColorProcessingRGB(Color, CSMtx, EOTF);
+}
+
+
 // shader to convert a NV12 texture to RGBA (real NV12 texture - DX11)
 void NV12ConvertPS(
 	float4 InPosition : SV_POSITION,
@@ -64,9 +97,10 @@ void NV12ConvertPS(
 	float3 YUV;
 	YUV.x = NV12ConvertUB.SRV_Y.Sample(NV12ConvertUB.SamplerB, ScaledUV).x;
 	YUV.yz = NV12ConvertUB.SRV_UV.Sample(NV12ConvertUB.SamplerP, ScaledUV).xy;
+	YUV.yz = NV12ConvertUB.SwapChroma != 0u ? YUV.zy : YUV.yz;
 
-	OutColor.xyz = YuvToRgb(YUV, NV12ConvertUB.ColorTransform, NV12ConvertUB.SrgbToLinear);
-	OutColor.w = 1.0f;
+	// Output with solid alpha
+	OutColor = float4(ApplyColorProcessing(YUV, NV12ConvertUB.ColorTransform, NV12ConvertUB.CSTransform, (NV12ConvertUB.SrgbToLinear != 0u) ? 1 : 0), 1.0f); 
 }
 
 
@@ -91,9 +125,9 @@ void NV12ConvertAsBytesPS(
 
 	YUV.y = NV12ConvertAsBytesUB.Texture.Sample(NV12ConvertAsBytesUB.SamplerP, float2(UPos, ScaledUV.y)).x;
 	YUV.z = NV12ConvertAsBytesUB.Texture.Sample(NV12ConvertAsBytesUB.SamplerP, float2(VPos, ScaledUV.y)).x;
+	YUV.yz = NV12ConvertAsBytesUB.SwapChroma != 0u ? YUV.zy : YUV.yz;
 
-	OutColor.xyz = YuvToRgb(YUV, NV12ConvertAsBytesUB.ColorTransform, NV12ConvertAsBytesUB.SrgbToLinear);
-	OutColor.w = 1.0f;
+	OutColor = float4(ApplyColorProcessing(YUV, NV12ConvertAsBytesUB.ColorTransform, NV12ConvertAsBytesUB.CSTransform, (NV12ConvertAsBytesUB.SrgbToLinear != 0u) ? 1 : 0), 1.0f); 
 }
 
 
@@ -130,17 +164,9 @@ void RGBConvertPS(
 {
 	float2 ScaledUV = InUV * RGBConvertUB.UVScale;
 	float4 Color = RGBConvertUB.Texture.Sample(RGBConvertUB.Sampler, ScaledUV);
-	// Inverse EOTF
-	if (RGBConvertUB.EOTF != 0)
-	{
-		Color.xyz = (RGBConvertUB.EOTF == 1) ? sRGBToLinear(Color.xyz) : ST2084ToLinear(Color.xyz);
-	}
-	// Adjust colorspace
-	float4 Col4 = float4(Color.xyz, 1);
-	OutColor.x = dot(RGBConvertUB.CSTransform[0], Col4);
-	OutColor.y = dot(RGBConvertUB.CSTransform[1], Col4);
-	OutColor.z = dot(RGBConvertUB.CSTransform[2], Col4);
-	OutColor.a = dot(RGBConvertUB.CSTransform[3], Color);
+
+	OutColor = float4(ApplyColorProcessingRGB(Color.xyz, RGBConvertUB.CSTransform, RGBConvertUB.EOTF),
+					  dot(RGBConvertUB.CSTransform[3], Color));
 }
 
 static const float4 kYCoCgOffsets = float4(-0.50196078431373, -0.50196078431373, 0, 0);
@@ -161,18 +187,9 @@ void YCoCgConvertPS(
 	float Co = CoCgSY.x / Scale;
 	float Cg = CoCgSY.y / Scale;
 	float Y = CoCgSY.w;
-	float4 Col4 = float4(Y + Co - Cg, Y + Cg, Y - Co - Cg, 1.0);
+	float3 Col = float3(Y + Co - Cg, Y + Cg, Y - Co - Cg);
 
-	// Inverse EOTF
-	if (YCoCgConvertUB.EOTF != 0)
-	{
-		Col4.xyz = (YCoCgConvertUB.EOTF == 1) ? sRGBToLinear(Col4.xyz) : ST2084ToLinear(Col4.xyz);
-	}
-	// Adjust colorspace
-	OutColor.x = dot(YCoCgConvertUB.CSTransform[0], Col4);
-	OutColor.y = dot(YCoCgConvertUB.CSTransform[1], Col4);
-	OutColor.z = dot(YCoCgConvertUB.CSTransform[2], Col4);
-	OutColor.w = 1.0;
+	OutColor = float4(ApplyColorProcessingRGB(Col, YCoCgConvertUB.CSTransform, YCoCgConvertUB.EOTF), 1.0f);
 }
 
 // shader to convert a luma (Y) and chroma (CbCr) texture to RGBA
@@ -252,30 +269,7 @@ void YUVv216ConvertPS(
 	CbY0CrY1 = (YUVv216ConvertUB.SwapChroma != 0u) ? CbY0CrY1.zyxw : CbY0CrY1;						// Cb/Cr swap
 	float3 YUV = (uint(InUV.x * YUVv216ConvertUB.OutputDimX) & 1u) ? CbY0CrY1.wxz : CbY0CrY1.yxz;	// Decode to YCbCr for odd/even pixels
 
-	// YUV -> RGB (using a plain 4x4 matrix - not the common 3x3 & offset setup)
-	float3 Color;
-	float4 YUV4 = float4(YUV, 1);
-	Color.x = dot(YUVv216ConvertUB.ColorTransform[0], YUV4);
-	Color.y = dot(YUVv216ConvertUB.ColorTransform[1], YUV4);
-	Color.z = dot(YUVv216ConvertUB.ColorTransform[2], YUV4);
-
-	// Make sure we safely stay above zero
-	Color = max(Color, 0.0f);
-
-	// Go linear
-	if (YUVv216ConvertUB.EOTF != 0u)
-	{
-		Color = (YUVv216ConvertUB.EOTF == 1u) ? sRGBToLinear(Color) : ST2084ToLinear(Color);
-	}
-
-	// Adjust colorspace
-	float4 Col4 = float4(Color, 1);
-	Color.x = dot(YUVv216ConvertUB.CSTransform[0], Col4);
-	Color.y = dot(YUVv216ConvertUB.CSTransform[1], Col4);
-	Color.z = dot(YUVv216ConvertUB.CSTransform[2], Col4);
-
-	// Output with solid alpha
-	OutColor = float4(Color, 1.0f); 
+	OutColor = float4(ApplyColorProcessing(YUV, YUVv216ConvertUB.ColorTransform, YUVv216ConvertUB.CSTransform, YUVv216ConvertUB.EOTF), 1.0f); 
 }
 
 
@@ -324,30 +318,7 @@ void YUVv210ConvertPS(
 		YUV = float3(In3.z, In2.z, In3.y);
 	}
 
-	// YUV -> RGB (using a plain 4x4 matrix - not the common 3x3 & offset setup)
-	float3 Color;
-	float4 YUV4 = float4(YUV, 1);
-	Color.x = dot(YUVv210ConvertUB.ColorTransform[0], YUV4);
-	Color.y = dot(YUVv210ConvertUB.ColorTransform[1], YUV4);
-	Color.z = dot(YUVv210ConvertUB.ColorTransform[2], YUV4);
-
-	// Make sure we safely stay above zero
-	Color = max(Color, 0.0f);
-
-	// Go linear
-	if (YUVv210ConvertUB.EOTF != 0u)
-	{
-		Color = (YUVv210ConvertUB.EOTF == 1u) ? sRGBToLinear(Color) : ST2084ToLinear(Color);
-	}
-
-	// Adjust colorspace
-	float4 Col4 = float4(Color, 1);
-	Color.x = dot(YUVv210ConvertUB.CSTransform[0], Col4);
-	Color.y = dot(YUVv210ConvertUB.CSTransform[1], Col4);
-	Color.z = dot(YUVv210ConvertUB.CSTransform[2], Col4);
-
-	// Output with solid alpha
-	OutColor = float4(Color, 1.0f); 
+	OutColor = float4(ApplyColorProcessing(YUV, YUVv210ConvertUB.ColorTransform, YUVv210ConvertUB.CSTransform, YUVv210ConvertUB.EOTF), 1.0f); 
 }
 
 
@@ -360,32 +331,7 @@ void YUVY416ConvertPS(
 	float4 AVYU = YUVY416ConvertUB.SRV_Y.Sample(YUVY416ConvertUB.BilinearClampedSamplerUV, InUV);
 	AVYU = (YUVY416ConvertUB.bIsARGBFmt != 0u) ? AVYU.zyxw : AVYU;	// PF_ARGB8 vs. PF_ABGR16 ordering
 
-	float3 YUV = AVYU.yzw;
-
-	// YUV -> RGB (using a plain 4x4 matrix - not the common 3x3 & offset setup)
-	float3 Color;
-	float4 YUV4 = float4(YUV, 1);
-	Color.x = dot(YUVY416ConvertUB.ColorTransform[0], YUV4);
-	Color.y = dot(YUVY416ConvertUB.ColorTransform[1], YUV4);
-	Color.z = dot(YUVY416ConvertUB.ColorTransform[2], YUV4);
-
-	// Make sure we safely stay above zero
-	Color = max(Color, 0.0f);
-
-	// Inverse EOTF
-	if (YUVY416ConvertUB.EOTF != 0)
-	{
-		Color = (YUVY416ConvertUB.EOTF == 1) ? sRGBToLinear(Color) : ST2084ToLinear(Color);
-	}
-
-	// Adjust colorspace
-	float4 Col4 = float4(Color, 1);
-	Color.x = dot(YUVY416ConvertUB.CSTransform[0], Col4);
-	Color.y = dot(YUVY416ConvertUB.CSTransform[1], Col4);
-	Color.z = dot(YUVY416ConvertUB.CSTransform[2], Col4);
-
-	// Output with alpha from sample
-	OutColor = float4(Color, AVYU.x); 
+	OutColor = float4(ApplyColorProcessing(AVYU.yzw, YUVY416ConvertUB.ColorTransform, YUVY416ConvertUB.CSTransform, YUVY416ConvertUB.EOTF), AVYU.x); 
 }
 
 
@@ -401,20 +347,7 @@ void ARGB16BigConvertPS(
 
 	float3 Color = In.yzw;
 
-	// Inverse EOTF
-	if (ARGB16BigConvertUB.EOTF != 0)
-	{
-		Color = (ARGB16BigConvertUB.EOTF == 1) ? sRGBToLinear(Color) : ST2084ToLinear(Color);
-	}
-
-	// Adjust colorspace
-	float4 Col4 = float4(Color, 1);
-	Color.x = dot(ARGB16BigConvertUB.CSTransform[0], Col4);
-	Color.y = dot(ARGB16BigConvertUB.CSTransform[1], Col4);
-	Color.z = dot(ARGB16BigConvertUB.CSTransform[2], Col4);
-
-	// Output with alpha from sample
-	OutColor = float4(Color, In.x); 
+	OutColor = float4(ApplyColorProcessingRGB(Color, ARGB16BigConvertUB.CSTransform, ARGB16BigConvertUB.EOTF), In.x);
 }
 
 
@@ -485,27 +418,7 @@ void P010ConvertPS(
 	YUV.x = P010ConvertUB.SRV_Y.Sample(P010ConvertUB.SamplerB, ScaledUV).x;
 	YUV.yz = P010ConvertUB.SRV_UV.Sample(P010ConvertUB.SamplerP, ScaledUV).xy;
 
-	// YUV -> RGB (using a plain 4x4 matrix - not the common 3x3 & offset setup)
-	float3 Color;
-	float4 YUV4 = float4(YUV, 1);
-	Color.x = dot(P010ConvertUB.ColorTransform[0], YUV4);
-	Color.y = dot(P010ConvertUB.ColorTransform[1], YUV4);
-	Color.z = dot(P010ConvertUB.ColorTransform[2], YUV4);
-
-	// Make sure we safely stay above zero
-	Color = max(Color, 0.0f);
-
-	// Go linear
-	Color = P010ConvertUB.IsST2084 ? ST2084ToLinear(Color) : Rec709ToLinear(Color);
-
-	// Adjust colorspace
-	float4 Col4 = float4(Color, 1);
-	Color.x = dot(P010ConvertUB.CSTransform[0], Col4);
-	Color.y = dot(P010ConvertUB.CSTransform[1], Col4);
-	Color.z = dot(P010ConvertUB.CSTransform[2], Col4);
-
-	// Output with solid alpha
-	OutColor = float4(Color, 1.0f); 
+	OutColor = float4(ApplyColorProcessing(YUV, P010ConvertUB.ColorTransform, P010ConvertUB.CSTransform, P010ConvertUB.EOTF), 1.0f); 
 }
 
 // shader to convert a P010 texture to RGBA (P010 data, packed into a G16 texture; manually addressed)
@@ -530,27 +443,7 @@ void P010ConvertAsUINT16sPS(
 	YUV.y = P010ConvertAsUINT16sUB.Texture.Sample(P010ConvertAsUINT16sUB.SamplerP, float2(UPos, ScaledUV.y)).x;
 	YUV.z = P010ConvertAsUINT16sUB.Texture.Sample(P010ConvertAsUINT16sUB.SamplerP, float2(VPos, ScaledUV.y)).x;
 
-	// YUV -> RGB (using a plain 4x3 matrix - not the common 3x3 & offset setup)
-	float3 Color;
-	float4 YUV4 = float4(YUV, 1);
-	Color.x = dot(P010ConvertAsUINT16sUB.ColorTransform[0], YUV4);
-	Color.y = dot(P010ConvertAsUINT16sUB.ColorTransform[1], YUV4);
-	Color.z = dot(P010ConvertAsUINT16sUB.ColorTransform[2], YUV4);
-
-	// Make sure we safely stay above zero
-	Color = max(Color, 0.0f);
-
-	// Go linear
-	Color = P010ConvertAsUINT16sUB.IsST2084 ? ST2084ToLinear(Color) : Rec709ToLinear(Color);
-
-	// Adjust colorspace
-	float4 Col4 = float4(Color, 1);
-	Color.x = dot(P010ConvertAsUINT16sUB.CSTransform[0], Col4);
-	Color.y = dot(P010ConvertAsUINT16sUB.CSTransform[1], Col4);
-	Color.z = dot(P010ConvertAsUINT16sUB.CSTransform[2], Col4);
-
-	// Output with solid alpha
-	OutColor = float4(Color, 1.0f); 
+	OutColor = float4(ApplyColorProcessing(YUV, P010ConvertAsUINT16sUB.ColorTransform, P010ConvertAsUINT16sUB.CSTransform, P010ConvertAsUINT16sUB.EOTF), 1.0f); 
 }
 
 
@@ -578,27 +471,7 @@ void P010_2101010ConvertPS(
 	YUV.y = dot(P010_2101010ConvertUB.SRV_U.Sample(P010_2101010ConvertUB.SamplerP, ScaledUV_UV).xyz, SelUV);
 	YUV.z = dot(P010_2101010ConvertUB.SRV_V.Sample(P010_2101010ConvertUB.SamplerP, ScaledUV_UV).xyz, SelUV);
 
-	// YUV -> RGB (using a plain 4x4 matrix - not the common 3x3 & offset setup)
-	float3 Color;
-	float4 YUV4 = float4(YUV, 1);
-	Color.x = dot(P010_2101010ConvertUB.ColorTransform[0], YUV4);
-	Color.y = dot(P010_2101010ConvertUB.ColorTransform[1], YUV4);
-	Color.z = dot(P010_2101010ConvertUB.ColorTransform[2], YUV4);
-
-	// Make sure we safely stay above zero
-	Color = max(Color, 0.0f);
-
-	// Go linear
-	Color = P010_2101010ConvertUB.IsST2084 ? ST2084ToLinear(Color) : Rec709ToLinear(Color);
-
-	// Adjust colorspace
-	float4 Col4 = float4(Color, 1);
-	Color.x = dot(P010_2101010ConvertUB.CSTransform[0], Col4);
-	Color.y = dot(P010_2101010ConvertUB.CSTransform[1], Col4);
-	Color.z = dot(P010_2101010ConvertUB.CSTransform[2], Col4);
-
-	// Output with solid alpha
-	OutColor = float4(Color, 1.0f); 
+	OutColor = float4(ApplyColorProcessing(YUV, P010_2101010ConvertUB.ColorTransform, P010_2101010ConvertUB.CSTransform, P010_2101010ConvertUB.EOTF), 1.0f); 
 #else
 	OutColor = float4(1,0,1,1);
 #endif // COMPILER_GLSL_ES3_1
diff --git a/Engine/Source/Runtime/Media/Public/IMediaTextureSample.h b/Engine/Source/Runtime/Media/Public/IMediaTextureSample.h
--- a/Engine/Source/Runtime/Media/Public/IMediaTextureSample.h
+++ b/Engine/Source/Runtime/Media/Public/IMediaTextureSample.h
@@ -378,10 +378,20 @@ public:
 	virtual FMatrix44f GetSampleToRGBMatrix() const
 	{
 		FMatrix Pre = FMatrix::Identity;
-		Pre.M[0][3] = -MediaShaders::YUVOffset8bits.X;
-		Pre.M[1][3] = -MediaShaders::YUVOffset8bits.Y;
-		Pre.M[2][3] = -MediaShaders::YUVOffset8bits.Z;
-		return FMatrix44f(MediaShaders::YuvToRgbRec709Scaled * Pre);
+		FVector Off;
+		switch (GetFormat())
+		{
+			case EMediaTextureSampleFormat::R4FL:		Off = MediaShaders::YUVOffsetFloat; break;
+			case EMediaTextureSampleFormat::Y416:
+			case EMediaTextureSampleFormat::P010:
+			case EMediaTextureSampleFormat::YUVv216:	Off = MediaShaders::YUVOffset16bits; break;
+			case EMediaTextureSampleFormat::YUVv210:	Off = MediaShaders::YUVOffset10bits; break;
+			default:									Off = MediaShaders::YUVOffset8bits; break;
+		}
+		Pre.M[0][3] = -Off.X;
+		Pre.M[1][3] = -Off.Y;
+		Pre.M[2][3] = -Off.Z;
+		return FMatrix44f(MediaShaders::YuvToRgbRec709Scaled * Pre);	// assumes sRGB & video range
 	}
 
 	/**
diff --git a/Engine/Source/Runtime/MediaAssets/Private/Misc/MediaTextureResource.cpp b/Engine/Source/Runtime/MediaAssets/Private/Misc/MediaTextureResource.cpp
--- a/Engine/Source/Runtime/MediaAssets/Private/Misc/MediaTextureResource.cpp
+++ b/Engine/Source/Runtime/MediaAssets/Private/Misc/MediaTextureResource.cpp
@@ -937,6 +937,18 @@ void FMediaTextureResource::ConvertSample(const TSharedPtr<IMediaTextureSample,
 }
 
 
+void FMediaTextureResource::GetColorSpaceConversionMatrixForSample(const TSharedPtr<IMediaTextureSample, ESPMode::ThreadSafe> Sample, FMatrix44f& ColorSpaceMtx)
+{
+	const UE::Color::FColorSpace& Working = UE::Color::FColorSpace::GetWorking();
+	ColorSpaceMtx = UE::Color::Transpose<float>(Working.GetXYZToRgb()) * Sample->GetGamutToXYZMatrix();
+	float NF = Sample->GetHDRNitsNormalizationFactor();
+	if (NF != 1.0f)
+	{
+		ColorSpaceMtx = ColorSpaceMtx.ApplyScale(NF);
+	}
+}
+
+
  void FMediaTextureResource::ConvertTextureToOutput(FRHITexture* InputTexture, const TSharedPtr<IMediaTextureSample, ESPMode::ThreadSafe>& Sample)
  {
 	// perform the conversion
@@ -976,124 +988,115 @@ void FMediaTextureResource::ConvertSample(const TSharedPtr<IMediaTextureSample,
 			GraphicsPSOInit.BoundShaderState.VertexDeclarationRHI = GMediaVertexDeclaration.VertexDeclarationRHI;
 			GraphicsPSOInit.BoundShaderState.VertexShaderRHI = VertexShader.GetVertexShader();
 
-			FMatrix YUVToRGBMatrix = Sample->GetYUVToRGBMatrix();
-			FVector YUVOffset;
-			YUVOffset = Sample->GetFullRange() ? MediaShaders::YUVOffsetNoScale8bits : MediaShaders::YUVOffset8bits;
-
-			bool bIsSampleOutputSrgb = Sample->IsOutputSrgb();
-
 			// Temporary SRV variables to hold references for the draw
 			FShaderResourceViewRHIRef TempSRV0, TempSRV1;
 
 			// Use the sample format to choose the conversion path
-			switch (Sample->GetFormat())
+			EMediaTextureSampleFormat SampleFormat = Sample->GetFormat();
+			switch (SampleFormat)
 			{
-				case EMediaTextureSampleFormat::CharBMP:
-				{
-					// Simple 1:1 copy plus flip & color adjustment (but using normal texture sampler: sRGB conversions may occur depending on setup; any manual sRGB/linear conversion is disabled)
-					TShaderMapRef<FBMPConvertPS> ConvertShader(ShaderMap);
-					GraphicsPSOInit.BoundShaderState.PixelShaderRHI = ConvertShader.GetPixelShader();
-					SetGraphicsPipelineState(CommandList, GraphicsPSOInit, 0);
-					SetShaderParametersLegacyPS(CommandList, ConvertShader, InputTexture, OutputDim, false);
-				}
-				break;
-
+				// Planar 8-bit YCbCr 420
 				case EMediaTextureSampleFormat::CharNV12:
+				case EMediaTextureSampleFormat::CharNV21:	// swapped CbCr (vs. NV12)
 				{
+					auto YUVMtx = Sample->GetSampleToRGBMatrix();
+					FMatrix44f ColorSpaceMtx;
+					GetColorSpaceConversionMatrixForSample(Sample, ColorSpaceMtx);
+					
 					if (InputTexture->GetFormat() == PF_NV12)
 					{
 						TShaderMapRef<FNV12ConvertPS> ConvertShader(ShaderMap);
 						GraphicsPSOInit.BoundShaderState.PixelShaderRHI = ConvertShader.GetPixelShader();
 						SetGraphicsPipelineState(CommandList, GraphicsPSOInit, 0);
 						FIntPoint TexDim = InputTexture->GetSizeXY();
-						TempSRV0 = RHICreateShaderResourceView(InputTexture, 0, 1, PF_G8); // note: the types of the views select the correct planes (offset) "magically"
+						TempSRV0 = RHICreateShaderResourceView(InputTexture, 0, 1, PF_G8);		// note: the types of the views select the correct planes (offset) "magically"
 						TempSRV1 = RHICreateShaderResourceView(InputTexture, 0, 1, PF_R8G8);
-						SetShaderParametersLegacyPS(CommandList, ConvertShader, TexDim, TempSRV0, TempSRV1, OutputDim, YUVToRGBMatrix, YUVOffset, bIsSampleOutputSrgb);
+						SetShaderParametersLegacyPS(CommandList, ConvertShader, TexDim, TempSRV0, TempSRV1, OutputDim, YUVMtx, Sample->GetEncodingType() == UE::Color::EEncoding::sRGB, ColorSpaceMtx, SampleFormat == EMediaTextureSampleFormat::CharNV21);
 					}
 					else
 					{
 						TShaderMapRef<FNV12ConvertAsBytesPS> ConvertShader(ShaderMap);
 						GraphicsPSOInit.BoundShaderState.PixelShaderRHI = ConvertShader.GetPixelShader();
 						SetGraphicsPipelineState(CommandList, GraphicsPSOInit, 0);
-						SetShaderParametersLegacyPS(CommandList, ConvertShader, InputTexture, OutputDim, YUVToRGBMatrix, YUVOffset, bIsSampleOutputSrgb);
+						SetShaderParametersLegacyPS(CommandList, ConvertShader, InputTexture, OutputDim, YUVMtx, Sample->GetEncodingType() == UE::Color::EEncoding::sRGB, ColorSpaceMtx, SampleFormat == EMediaTextureSampleFormat::CharNV21);
 					}
 				}
 				break;
 
-				case EMediaTextureSampleFormat::CharNV21:
+				// Planar 10-bit YCbCr 420 (in 16-bit components)
+				case EMediaTextureSampleFormat::P010:
 				{
-					// source texture might be NV12 or G8...
-					TShaderMapRef<FNV21ConvertPS> ConvertShader(ShaderMap);
-					GraphicsPSOInit.BoundShaderState.PixelShaderRHI = ConvertShader.GetPixelShader();
-					SetGraphicsPipelineState(CommandList, GraphicsPSOInit, 0);
-					SetShaderParametersLegacyPS(CommandList, ConvertShader, InputTexture, OutputDim, YUVToRGBMatrix, YUVOffset, bIsSampleOutputSrgb);
-				}
-				break;
+					auto YUVMtx = Sample->GetSampleToRGBMatrix();
+					FMatrix44f ColorSpaceMtx;
+					GetColorSpaceConversionMatrixForSample(Sample, ColorSpaceMtx);
 
-				case EMediaTextureSampleFormat::CharYVYU:
-				{
-					TShaderMapRef<FYVYUConvertPS> ConvertShader(ShaderMap);
-					GraphicsPSOInit.BoundShaderState.PixelShaderRHI = ConvertShader.GetPixelShader();
-					SetGraphicsPipelineState(CommandList, GraphicsPSOInit, 0);
-					SetShaderParametersLegacyPS(CommandList, ConvertShader, InputTexture, YUVToRGBMatrix, YUVOffset, bIsSampleOutputSrgb);
+					FIntPoint TexDim = InputTexture->GetSizeXY();
+					if (InputTexture->GetFormat() == PF_P010)
+					{
+						TShaderMapRef<FP010ConvertPS> ConvertShader(ShaderMap);
+						GraphicsPSOInit.BoundShaderState.PixelShaderRHI = ConvertShader.GetPixelShader();
+						SetGraphicsPipelineState(CommandList, GraphicsPSOInit, 0);
+
+						FShaderResourceViewRHIRef Y_SRV = RHICreateShaderResourceView(InputTexture, 0, 1, PF_G16);		// note: the types of the views select the correct planes (offset) "magically"
+						FShaderResourceViewRHIRef UV_SRV = RHICreateShaderResourceView(InputTexture, 0, 1, PF_G16R16);
+						SetShaderParametersLegacyPS(CommandList, ConvertShader, TexDim, Y_SRV, UV_SRV, OutputDim, YUVMtx, ColorSpaceMtx, Sample->GetEncodingType() == UE::Color::EEncoding::sRGB, Sample->GetEncodingType() == UE::Color::EEncoding::ST2084);
+					}
+					else
+					{
+						TShaderMapRef<FP010ConvertAsUINT16sPS> ConvertShader(ShaderMap);
+						GraphicsPSOInit.BoundShaderState.PixelShaderRHI = ConvertShader.GetPixelShader();
+						SetGraphicsPipelineState(CommandList, GraphicsPSOInit, 0);
+
+						SetShaderParametersLegacyPS(CommandList, ConvertShader, TexDim, InputTexture, OutputDim, YUVMtx, ColorSpaceMtx, Sample->GetEncodingType() == UE::Color::EEncoding::sRGB, Sample->GetEncodingType() == UE::Color::EEncoding::ST2084);
+					}
 				}
 				break;
 
+				// Various YCbCr 422 formats
 				case EMediaTextureSampleFormat::CharYUY2:		// Y0CbY1Cr
+				case EMediaTextureSampleFormat::CharYVYU:		// Y0CrY1Cb
 				case EMediaTextureSampleFormat::Char2VUY:		// CbY0CrY1
 				case EMediaTextureSampleFormat::YUVv216:		// CbY0CrY1
 				case EMediaTextureSampleFormat::CharUYVY:		// CbY0CrY1
 				{
-					// Get the complete matrix to convert sample data to RGB
 					auto YUVMtx = Sample->GetSampleToRGBMatrix();
-
-					// Setup conversion from Rec2020 to current working color space
-					const UE::Color::FColorSpace& Working = UE::Color::FColorSpace::GetWorking();
-					FMatrix44f ColorSpaceMtx = UE::Color::Transpose<float>(Working.GetXYZToRgb()) * Sample->GetGamutToXYZMatrix();
-					// Normalize output (e.g. 80 or 100 nits == 1.0)
-					ColorSpaceMtx = ColorSpaceMtx.ApplyScale(Sample->GetHDRNitsNormalizationFactor());
+					FMatrix44f ColorSpaceMtx;
+					GetColorSpaceConversionMatrixForSample(Sample, ColorSpaceMtx);
 
 					TShaderMapRef<FYUVv216ConvertPS> ConvertShader(ShaderMap);
 					GraphicsPSOInit.BoundShaderState.PixelShaderRHI = ConvertShader.GetPixelShader();
 					SetGraphicsPipelineState(CommandList, GraphicsPSOInit, 0);
 					SetShaderParametersLegacyPS(CommandList, ConvertShader, InputTexture, OutputDim, YUVMtx, Sample->GetEncodingType() == UE::Color::EEncoding::sRGB, Sample->GetEncodingType() == UE::Color::EEncoding::ST2084, ColorSpaceMtx,
-												Sample->GetFormat() != EMediaTextureSampleFormat::CharYUY2,	// Y or Cb first
-												InputTexture->GetFormat() == PF_B8G8R8A8,					// ARGB vs. ABGR (memory order)
-												false														// No Cb / Cr swap
+												SampleFormat!= EMediaTextureSampleFormat::CharYUY2 && SampleFormat != EMediaTextureSampleFormat::CharYVYU,	// Y or Cb first
+												InputTexture->GetFormat() == PF_B8G8R8A8,																	// ARGB vs. ABGR (memory order)
+												SampleFormat == EMediaTextureSampleFormat::CharYVYU															// Cb / Cr swap
 												);
 				}
 				break;
 
+				// 10-bit YCbCr 422 format (stored in A2R10G10B10)
 				case EMediaTextureSampleFormat::YUVv210:
 				{
-					// Get the complete matrix to convert sample data to RGB
 					auto YUVMtx = Sample->GetSampleToRGBMatrix();
-
-					// Setup conversion from Rec2020 to current working color space
-					const UE::Color::FColorSpace& Working = UE::Color::FColorSpace::GetWorking();
-					FMatrix44f ColorSpaceMtx = UE::Color::Transpose<float>(Working.GetXYZToRgb()) * Sample->GetGamutToXYZMatrix();
-					// Normalize output (e.g. 80 or 100 nits == 1.0)
-					ColorSpaceMtx = ColorSpaceMtx.ApplyScale(Sample->GetHDRNitsNormalizationFactor());
+					FMatrix44f ColorSpaceMtx;
+					GetColorSpaceConversionMatrixForSample(Sample, ColorSpaceMtx);
 
 					TShaderMapRef<FYUVv210ConvertPS> ConvertShader(ShaderMap);
 					GraphicsPSOInit.BoundShaderState.PixelShaderRHI = ConvertShader.GetPixelShader();
 					SetGraphicsPipelineState(CommandList, GraphicsPSOInit, 0);
-					SetShaderParametersLegacyPS(CommandList, ConvertShader, InputTexture, OutputDim, YUVMtx, Sample->GetEncodingType() == UE::Color::EEncoding::sRGB, Sample->GetEncodingType() == UE::Color::EEncoding::ST2084, ColorSpaceMtx, Sample->GetFormat() != EMediaTextureSampleFormat::CharYUY2);
+					SetShaderParametersLegacyPS(CommandList, ConvertShader, InputTexture, OutputDim, YUVMtx, Sample->GetEncodingType() == UE::Color::EEncoding::sRGB, Sample->GetEncodingType() == UE::Color::EEncoding::ST2084, ColorSpaceMtx,
+												true);
 				}
 				break;
 
+				// Various YCbCR 444 formats
 				case EMediaTextureSampleFormat::CharAYUV:
 				case EMediaTextureSampleFormat::Y416:
 				case EMediaTextureSampleFormat::R4FL:
 				{
-					// Get the complete matrix to convert sample data to RGB
 					auto YUVMtx = Sample->GetSampleToRGBMatrix();
-
-					// Setup conversion from Rec2020 to current working color space
-					const UE::Color::FColorSpace& Working = UE::Color::FColorSpace::GetWorking();
-					FMatrix44f ColorSpaceMtx = UE::Color::Transpose<float>(Working.GetXYZToRgb()) * Sample->GetGamutToXYZMatrix();
-					// Normalize output (e.g. 80 or 100 nits == 1.0)
-					ColorSpaceMtx = ColorSpaceMtx.ApplyScale(Sample->GetHDRNitsNormalizationFactor());
+					FMatrix44f ColorSpaceMtx;
+					GetColorSpaceConversionMatrixForSample(Sample, ColorSpaceMtx);
 
 					TShaderMapRef<FYUVY416ConvertPS> ConvertShader(ShaderMap);
 					GraphicsPSOInit.BoundShaderState.PixelShaderRHI = ConvertShader.GetPixelShader();
@@ -1104,6 +1107,7 @@ void FMediaTextureResource::ConvertSample(const TSharedPtr<IMediaTextureSample,
 				}
 				break;
 
+				// Various RGB(A) formats
 				case EMediaTextureSampleFormat::CharBGR10A2:
 				case EMediaTextureSampleFormat::CharRGBA:
 				case EMediaTextureSampleFormat::CharBGRA:
@@ -1130,9 +1134,8 @@ void FMediaTextureResource::ConvertSample(const TSharedPtr<IMediaTextureSample,
 						bConvertFromST2084 = (Sample->GetEncodingType() == UE::Color::EEncoding::ST2084);
 					}
 
-					const UE::Color::FColorSpace& Working = UE::Color::FColorSpace::GetWorking();
-					FMatrix44f ColorSpaceMtx = UE::Color::Transpose<float>(Working.GetXYZToRgb()) * Sample->GetGamutToXYZMatrix();
-					ColorSpaceMtx = ColorSpaceMtx.ApplyScale(Sample->GetHDRNitsNormalizationFactor());
+					FMatrix44f ColorSpaceMtx;
+					GetColorSpaceConversionMatrixForSample(Sample, ColorSpaceMtx);
 
 					TShaderMapRef<FRGBConvertPS> ConvertShader(ShaderMap);
 					GraphicsPSOInit.BoundShaderState.PixelShaderRHI = ConvertShader.GetPixelShader();
@@ -1143,9 +1146,8 @@ void FMediaTextureResource::ConvertSample(const TSharedPtr<IMediaTextureSample,
 
 				case EMediaTextureSampleFormat::ARGB16_BIG:
 				{
-					const UE::Color::FColorSpace& Working = UE::Color::FColorSpace::GetWorking();
-					FMatrix44f ColorSpaceMtx = UE::Color::Transpose<float>(Working.GetXYZToRgb()) * Sample->GetGamutToXYZMatrix();
-					ColorSpaceMtx = ColorSpaceMtx.ApplyScale(Sample->GetHDRNitsNormalizationFactor());
+					FMatrix44f ColorSpaceMtx;
+					GetColorSpaceConversionMatrixForSample(Sample, ColorSpaceMtx);
 
 					TShaderMapRef<FARGB16BigConvertPS> ConvertShader(ShaderMap);
 					GraphicsPSOInit.BoundShaderState.PixelShaderRHI = ConvertShader.GetPixelShader();
@@ -1156,36 +1158,13 @@ void FMediaTextureResource::ConvertSample(const TSharedPtr<IMediaTextureSample,
 				}
 				break;
 
-				case EMediaTextureSampleFormat::P010:
+				case EMediaTextureSampleFormat::CharBMP:
 				{
-					// Get the complete matrix to convert sample data to RGB
-					auto YUVMtx = Sample->GetSampleToRGBMatrix();
-
-					// Setup conversion from Rec2020 to current working color space
-					const UE::Color::FColorSpace& Working = UE::Color::FColorSpace::GetWorking();
-					FMatrix44f ColorSpaceMtx = UE::Color::Transpose<float>(Working.GetXYZToRgb()) * Sample->GetGamutToXYZMatrix();
-					// Normalize output (e.g. 80 or 100 nits == 1.0)
-					ColorSpaceMtx = ColorSpaceMtx.ApplyScale(Sample->GetHDRNitsNormalizationFactor());
-
-					FIntPoint TexDim = InputTexture->GetSizeXY();
-					if (InputTexture->GetFormat() == PF_P010)
-					{
-						TShaderMapRef<FP010ConvertPS> ConvertShader(ShaderMap);
-						GraphicsPSOInit.BoundShaderState.PixelShaderRHI = ConvertShader.GetPixelShader();
-						SetGraphicsPipelineState(CommandList, GraphicsPSOInit, 0);
-
-						FShaderResourceViewRHIRef Y_SRV = RHICreateShaderResourceView(InputTexture, 0, 1, PF_G16);
-						FShaderResourceViewRHIRef UV_SRV = RHICreateShaderResourceView(InputTexture, 0, 1, PF_G16R16);
-						SetShaderParametersLegacyPS(CommandList, ConvertShader, TexDim, Y_SRV, UV_SRV, OutputDim, YUVMtx, ColorSpaceMtx, Sample->GetEncodingType() == UE::Color::EEncoding::ST2084);
-					}
-					else
-					{
-						TShaderMapRef<FP010ConvertAsUINT16sPS> ConvertShader(ShaderMap);
-						GraphicsPSOInit.BoundShaderState.PixelShaderRHI = ConvertShader.GetPixelShader();
-						SetGraphicsPipelineState(CommandList, GraphicsPSOInit, 0);
-
-						SetShaderParametersLegacyPS(CommandList, ConvertShader, TexDim, InputTexture, OutputDim, YUVMtx, ColorSpaceMtx, Sample->GetEncodingType() == UE::Color::EEncoding::ST2084);
-					}
+					// Simple 1:1 copy plus flip & color adjustment (but using normal texture sampler: sRGB conversions may occur depending on setup; any manual sRGB/linear conversion is disabled)
+					TShaderMapRef<FBMPConvertPS> ConvertShader(ShaderMap);
+					GraphicsPSOInit.BoundShaderState.PixelShaderRHI = ConvertShader.GetPixelShader();
+					SetGraphicsPipelineState(CommandList, GraphicsPSOInit, 0);
+					SetShaderParametersLegacyPS(CommandList, ConvertShader, InputTexture, OutputDim, false);
 				}
 				break;
 
@@ -1207,8 +1186,8 @@ void FMediaTextureResource::ConvertSample(const TSharedPtr<IMediaTextureSample,
 						bConvertFromSRGB = (Sample->GetEncodingType() == UE::Color::EEncoding::sRGB);
 					}
 
-					const UE::Color::FColorSpace& Working = UE::Color::FColorSpace::GetWorking();
-					FMatrix44f ColorSpaceMtx = UE::Color::Transpose<float>(Working.GetXYZToRgb()) * Sample->GetGamutToXYZMatrix();
+					FMatrix44f ColorSpaceMtx;
+					GetColorSpaceConversionMatrixForSample(Sample, ColorSpaceMtx);
 
 					// Convert the DXn compressed data into RGBA so we can more easily process it using the existing logic (e.g mip generation)
 					// (this also will deal with any sRGB/Rec703 conversions - as well as any color space conversions)
@@ -1222,7 +1201,7 @@ void FMediaTextureResource::ConvertSample(const TSharedPtr<IMediaTextureSample,
 				case EMediaTextureSampleFormat::BC4:
 				{
 					// We are using the RGB conversion shader to just pipe BC4 data (read into R) and remap it to be usable with L8 or A8 etc. (RRRR) using the otherwise idle CS conversion matrix
-					// (note: no EOTF - this is alpha only!)
+					// (note: no EOTF or CS - this is alpha only!)
 					TShaderMapRef<FRGBConvertPS> ConvertShader(ShaderMap);
 					GraphicsPSOInit.BoundShaderState.PixelShaderRHI = ConvertShader.GetPixelShader();
 					SetGraphicsPipelineState(CommandList, GraphicsPSOInit, 0);
@@ -1245,8 +1224,8 @@ void FMediaTextureResource::ConvertSample(const TSharedPtr<IMediaTextureSample,
 						bConvertFromSRGB = (Sample->GetEncodingType() == UE::Color::EEncoding::sRGB);
 					}
 
-					const UE::Color::FColorSpace& Working = UE::Color::FColorSpace::GetWorking();
-					FMatrix44f ColorSpaceMtx = UE::Color::Transpose<float>(Working.GetXYZToRgb()) * Sample->GetGamutToXYZMatrix();
+					FMatrix44f ColorSpaceMtx;
+					GetColorSpaceConversionMatrixForSample(Sample, ColorSpaceMtx);
 
 					// Convert the DXn compressed YCoCg data into RGBA so we can more easily process it using the existing logic (e.g mip generation)
 					// (this also will deal with any sRGB/Rec703 conversions - as well as any color space conversions)
diff --git a/Engine/Source/Runtime/MediaAssets/Private/Misc/MediaTextureResource.h b/Engine/Source/Runtime/MediaAssets/Private/Misc/MediaTextureResource.h
--- a/Engine/Source/Runtime/MediaAssets/Private/Misc/MediaTextureResource.h
+++ b/Engine/Source/Runtime/MediaAssets/Private/Misc/MediaTextureResource.h
@@ -193,6 +193,9 @@ protected:
 	bool RequiresConversion(const TSharedPtr<IMediaTextureSample, ESPMode::ThreadSafe>& Sample, uint8 numMips) const;
 	bool RequiresConversion(const FTexture2DRHIRef& SampleTexture, const FIntPoint & OutputDim, uint8 numMips) const;
 
+	/** Compute CS conversion martix based on sample's data */
+	void GetColorSpaceConversionMatrixForSample(const TSharedPtr<IMediaTextureSample, ESPMode::ThreadSafe> Sample, FMatrix44f& ColorSpaceMtx);
+
 private:
 
 	/** Platform uses GL/ES ImageExternal */
diff --git a/Engine/Source/Runtime/RenderCore/Private/MediaShaders.cpp b/Engine/Source/Runtime/RenderCore/Private/MediaShaders.cpp
--- a/Engine/Source/Runtime/RenderCore/Private/MediaShaders.cpp
+++ b/Engine/Source/Runtime/RenderCore/Private/MediaShaders.cpp
@@ -104,6 +104,9 @@ namespace MediaShaders
 	const FVector YUVOffset16bits = FVector(4096.0f / 65535.0f, 32768.0f / 65535.0f, 32768.0f / 65535.0f);
 	const FVector YUVOffsetNoScale16bits = FVector(0.0f, 32768.0f / 65535.0f, 32768.0f / 65535.0f);
 
+	const FVector YUVOffsetFloat = FVector(1.0f / 16.0f, 0.5f, 0.5f);
+	const FVector YUVOffsetNoScaleFloat = FVector(0.0f, 0.5f, 0.5f);
+
 	/** Setup YUV Offset in matrix */
 	FMatrix CombineColorTransformAndOffset(const FMatrix& InMatrix, const FVector& InYUVOffset)
 	{
@@ -188,8 +191,10 @@ void FBMPConvertPS::SetParameters(FRHIBatchedShaderParameters& BatchedParameters
 
 BEGIN_GLOBAL_SHADER_PARAMETER_STRUCT(FNV12ConvertUB, )
 SHADER_PARAMETER(FMatrix44f, ColorTransform)
+SHADER_PARAMETER(FMatrix44f, CSTransform)
 SHADER_PARAMETER(uint32, OutputWidth)
 SHADER_PARAMETER(uint32, SrgbToLinear)
+SHADER_PARAMETER(uint32, SwapChroma)
 SHADER_PARAMETER(FVector2f, UVScale)
 SHADER_PARAMETER_SRV(Texture2D, SRV_Y)
 SHADER_PARAMETER_SRV(Texture2D, SRV_UV)
@@ -200,11 +205,12 @@ END_GLOBAL_SHADER_PARAMETER_STRUCT()
 IMPLEMENT_GLOBAL_SHADER_PARAMETER_STRUCT(FNV12ConvertUB, "NV12ConvertUB");
 IMPLEMENT_SHADER_TYPE(, FNV12ConvertPS, TEXT("/Engine/Private/MediaShaders.usf"), TEXT("NV12ConvertPS"), SF_Pixel);
 
-void FNV12ConvertPS::SetParameters(FRHIBatchedShaderParameters& BatchedParameters, const FIntPoint & TexDim, FShaderResourceViewRHIRef SRV_Y, FShaderResourceViewRHIRef SRV_UV, const FIntPoint& OutputDimensions, const FMatrix& ColorTransform, const FVector& YUVOffset, bool SrgbToLinear)
+void FNV12ConvertPS::SetParameters(FRHIBatchedShaderParameters& BatchedParameters, const FIntPoint & TexDim, FShaderResourceViewRHIRef SRV_Y, FShaderResourceViewRHIRef SRV_UV, const FIntPoint& OutputDimensions, const FMatrix44f& ColorTransform, bool SrgbToLinear, const FMatrix44f& CSTransform, bool bSwapChroma)
 {
 	FNV12ConvertUB UB;
 	{
-		UB.ColorTransform = (FMatrix44f)MediaShaders::CombineColorTransformAndOffset(ColorTransform, YUVOffset);
+		UB.ColorTransform = ColorTransform;
+		UB.CSTransform = CSTransform;
 		UB.OutputWidth = OutputDimensions.X;
 		UB.SamplerB = TStaticSamplerState<SF_Bilinear>::GetRHI();
 		UB.SamplerP = TStaticSamplerState<SF_Point>::GetRHI();
@@ -212,6 +218,7 @@ void FNV12ConvertPS::SetParameters(FRHIBatchedShaderParameters& BatchedParameter
 		UB.SRV_Y = SRV_Y;
 		UB.SRV_UV = SRV_UV;
 		UB.UVScale = FVector2f((float)OutputDimensions.X / (float)TexDim.X, (float)OutputDimensions.Y / (float)TexDim.Y);
+		UB.SwapChroma = bSwapChroma;
 	}
 
 	TUniformBufferRef<FNV12ConvertUB> Data = TUniformBufferRef<FNV12ConvertUB>::CreateUniformBufferImmediate(UB, UniformBuffer_SingleFrame);
@@ -224,8 +231,10 @@ void FNV12ConvertPS::SetParameters(FRHIBatchedShaderParameters& BatchedParameter
 
 BEGIN_GLOBAL_SHADER_PARAMETER_STRUCT(FNV12ConvertAsBytesUB, )
 SHADER_PARAMETER(FMatrix44f, ColorTransform)
+SHADER_PARAMETER(FMatrix44f, CSTransform)
 SHADER_PARAMETER(uint32, OutputWidth)
 SHADER_PARAMETER(uint32, SrgbToLinear)
+SHADER_PARAMETER(uint32, SwapChroma)
 SHADER_PARAMETER(FVector2f, UVScale)
 SHADER_PARAMETER_TEXTURE(Texture2D, Texture)
 SHADER_PARAMETER_SAMPLER(SamplerState, SamplerB)
@@ -235,17 +244,19 @@ END_GLOBAL_SHADER_PARAMETER_STRUCT()
 IMPLEMENT_GLOBAL_SHADER_PARAMETER_STRUCT(FNV12ConvertAsBytesUB, "NV12ConvertAsBytesUB");
 IMPLEMENT_SHADER_TYPE(, FNV12ConvertAsBytesPS, TEXT("/Engine/Private/MediaShaders.usf"), TEXT("NV12ConvertAsBytesPS"), SF_Pixel);
 
-void FNV12ConvertAsBytesPS::SetParameters(FRHIBatchedShaderParameters& BatchedParameters, TRefCountPtr<FRHITexture2D> NV12Texture, const FIntPoint& OutputDimensions, const FMatrix& ColorTransform, const FVector& YUVOffset, bool SrgbToLinear)
+void FNV12ConvertAsBytesPS::SetParameters(FRHIBatchedShaderParameters& BatchedParameters, TRefCountPtr<FRHITexture2D> NV12Texture, const FIntPoint& OutputDimensions, const FMatrix44f& ColorTransform, bool SrgbToLinear, const FMatrix44f& CSTransform, bool bSwapChroma)
 {
 	FNV12ConvertAsBytesUB UB;
 	{
-		UB.ColorTransform = (FMatrix44f)MediaShaders::CombineColorTransformAndOffset(ColorTransform, YUVOffset);
+		UB.ColorTransform = ColorTransform;
+		UB.CSTransform = CSTransform;
 		UB.OutputWidth = OutputDimensions.X;
 		UB.SamplerB = TStaticSamplerState<SF_Bilinear>::GetRHI();
 		UB.SamplerP = TStaticSamplerState<SF_Point>::GetRHI();
 		UB.SrgbToLinear = SrgbToLinear;
 		UB.Texture = NV12Texture;
 		UB.UVScale = FVector2f((float)OutputDimensions.X / (float)NV12Texture->GetSizeX(), (float)OutputDimensions.Y / (float)NV12Texture->GetSizeY());
+		UB.SwapChroma = bSwapChroma;
 	}
 
 	TUniformBufferRef<FNV12ConvertAsBytesUB> Data = TUniformBufferRef<FNV12ConvertAsBytesUB>::CreateUniformBufferImmediate(UB, UniformBuffer_SingleFrame);
@@ -294,7 +305,7 @@ void FNV21ConvertPS::SetParameters(FRHIBatchedShaderParameters& BatchedParameter
 BEGIN_GLOBAL_SHADER_PARAMETER_STRUCT(FP010ConvertUB, )
 SHADER_PARAMETER(FMatrix44f, ColorTransform)
 SHADER_PARAMETER(FMatrix44f, CSTransform)
-SHADER_PARAMETER(uint32, IsST2084)
+SHADER_PARAMETER(uint32, EOTF)
 SHADER_PARAMETER(FVector2f, UVScale)
 SHADER_PARAMETER_SRV(Texture2D, SRV_Y)
 SHADER_PARAMETER_SRV(Texture2D, SRV_UV)
@@ -305,7 +316,7 @@ END_GLOBAL_SHADER_PARAMETER_STRUCT()
 IMPLEMENT_GLOBAL_SHADER_PARAMETER_STRUCT(FP010ConvertUB, "P010ConvertUB");
 IMPLEMENT_SHADER_TYPE(, FP010ConvertPS, TEXT("/Engine/Private/MediaShaders.usf"), TEXT("P010ConvertPS"), SF_Pixel);
 
-void FP010ConvertPS::SetParameters(FRHIBatchedShaderParameters& BatchedParameters, const FIntPoint& TexDim, FShaderResourceViewRHIRef SRV_Y, FShaderResourceViewRHIRef SRV_UV, const FIntPoint& OutputDimensions, const FMatrix44f& ColorTransform, const FMatrix44f& CSTransform, bool bIsST2084)
+void FP010ConvertPS::SetParameters(FRHIBatchedShaderParameters& BatchedParameters, const FIntPoint& TexDim, FShaderResourceViewRHIRef SRV_Y, FShaderResourceViewRHIRef SRV_UV, const FIntPoint& OutputDimensions, const FMatrix44f& ColorTransform, const FMatrix44f& CSTransform, bool bSrgbToLinear, bool bIsST2084)
 {
 	FP010ConvertUB UB;
 	{
@@ -313,7 +324,7 @@ void FP010ConvertPS::SetParameters(FRHIBatchedShaderParameters& BatchedParameter
 		UB.CSTransform = CSTransform;
 		UB.SamplerB = TStaticSamplerState<SF_Bilinear>::GetRHI();
 		UB.SamplerP = TStaticSamplerState<SF_Point>::GetRHI();
-		UB.IsST2084 = bIsST2084;
+		UB.EOTF = bSrgbToLinear ? 1 : (bIsST2084 ? 2 : 0);
 		UB.SRV_Y = SRV_Y;
 		UB.SRV_UV = SRV_UV;
 		UB.UVScale = FVector2f((float)OutputDimensions.X / (float)TexDim.X, (float)OutputDimensions.Y / (float)TexDim.Y);
@@ -331,7 +342,7 @@ BEGIN_GLOBAL_SHADER_PARAMETER_STRUCT(FP010ConvertAsUINT16sUB, )
 SHADER_PARAMETER(FMatrix44f, ColorTransform)
 SHADER_PARAMETER(FMatrix44f, CSTransform)
 SHADER_PARAMETER(uint32, OutputWidth)
-SHADER_PARAMETER(uint32, IsST2084)
+SHADER_PARAMETER(uint32, EOTF)
 SHADER_PARAMETER(FVector2f, UVScale)
 SHADER_PARAMETER_TEXTURE(Texture2D, Texture)
 SHADER_PARAMETER_SAMPLER(SamplerState, SamplerB)
@@ -341,7 +352,7 @@ END_GLOBAL_SHADER_PARAMETER_STRUCT()
 IMPLEMENT_GLOBAL_SHADER_PARAMETER_STRUCT(FP010ConvertAsUINT16sUB, "P010ConvertAsUINT16sUB");
 IMPLEMENT_SHADER_TYPE(, FP010ConvertAsUINT16sPS, TEXT("/Engine/Private/MediaShaders.usf"), TEXT("P010ConvertAsUINT16sPS"), SF_Pixel);
 
-void FP010ConvertAsUINT16sPS::SetParameters(FRHIBatchedShaderParameters& BatchedParameters, const FIntPoint& TexDim, TRefCountPtr<FRHITexture2D> NV12Texture, const FIntPoint& OutputDimensions, const FMatrix44f& ColorTransform, const FMatrix44f& CSTransform, bool bIsST2084)
+void FP010ConvertAsUINT16sPS::SetParameters(FRHIBatchedShaderParameters& BatchedParameters, const FIntPoint& TexDim, TRefCountPtr<FRHITexture2D> NV12Texture, const FIntPoint& OutputDimensions, const FMatrix44f& ColorTransform, const FMatrix44f& CSTransform, bool bSrgbToLinear, bool bIsST2084)
 {
 	FP010ConvertAsUINT16sUB UB;
 	{
@@ -350,7 +361,7 @@ void FP010ConvertAsUINT16sPS::SetParameters(FRHIBatchedShaderParameters& Batched
 		UB.OutputWidth = OutputDimensions.X;
 		UB.SamplerB = TStaticSamplerState<SF_Bilinear>::GetRHI();
 		UB.SamplerP = TStaticSamplerState<SF_Point>::GetRHI();
-		UB.IsST2084 = bIsST2084;
+		UB.EOTF = bSrgbToLinear ? 1 : (bIsST2084 ? 2 : 0);
 		UB.Texture = NV12Texture;
 		UB.UVScale = FVector2f((float)OutputDimensions.X / (float)TexDim.X, (float)OutputDimensions.Y / (float)TexDim.Y);
 	}
@@ -366,7 +377,7 @@ void FP010ConvertAsUINT16sPS::SetParameters(FRHIBatchedShaderParameters& Batched
 BEGIN_GLOBAL_SHADER_PARAMETER_STRUCT(FP010_2101010ConvertUB, )
 SHADER_PARAMETER(FMatrix44f, ColorTransform)
 SHADER_PARAMETER(FMatrix44f, CSTransform)
-SHADER_PARAMETER(uint32, IsST2084)
+SHADER_PARAMETER(uint32, EOTF)
 SHADER_PARAMETER(FVector2f, UVScaleY)
 SHADER_PARAMETER(FVector2f, UVScaleUV)
 SHADER_PARAMETER_SRV(Texture2D, SRV_Y)
@@ -380,14 +391,14 @@ END_GLOBAL_SHADER_PARAMETER_STRUCT()
 IMPLEMENT_GLOBAL_SHADER_PARAMETER_STRUCT(FP010_2101010ConvertUB, "P010_2101010ConvertUB");
 IMPLEMENT_SHADER_TYPE(, FP010_2101010ConvertPS, TEXT("/Engine/Private/MediaShaders.usf"), TEXT("P010_2101010ConvertPS"), SF_Pixel);
 
-void FP010_2101010ConvertPS::SetParameters(FRHIBatchedShaderParameters& BatchedParameters, const FIntPoint& TexDim, FShaderResourceViewRHIRef SRV_Y, FShaderResourceViewRHIRef SRV_U, FShaderResourceViewRHIRef SRV_V, const FIntPoint& OutputDimensions, const FMatrix44f& ColorTransform, const FMatrix44f& CSTransform, bool bIsST2084)
+void FP010_2101010ConvertPS::SetParameters(FRHIBatchedShaderParameters& BatchedParameters, const FIntPoint& TexDim, FShaderResourceViewRHIRef SRV_Y, FShaderResourceViewRHIRef SRV_U, FShaderResourceViewRHIRef SRV_V, const FIntPoint& OutputDimensions, const FMatrix44f& ColorTransform, const FMatrix44f& CSTransform, bool bSrgbToLinear, bool bIsST2084)
 {
 	FP010_2101010ConvertUB UB;
 	{
 		UB.ColorTransform = ColorTransform;
 		UB.CSTransform = CSTransform;
 		UB.SamplerP = TStaticSamplerState<SF_Point>::GetRHI();
-		UB.IsST2084 = bIsST2084;
+		UB.EOTF = bSrgbToLinear ? 1 : (bIsST2084 ? 2 : 0);
 		UB.SRV_Y = SRV_Y;
 		UB.SRV_U = SRV_U;
 		UB.SRV_V = SRV_V;
diff --git a/Engine/Source/Runtime/RenderCore/Public/MediaShaders.h b/Engine/Source/Runtime/RenderCore/Public/MediaShaders.h
--- a/Engine/Source/Runtime/RenderCore/Public/MediaShaders.h
+++ b/Engine/Source/Runtime/RenderCore/Public/MediaShaders.h
@@ -80,6 +80,12 @@ namespace MediaShaders
 	/** YUV Offset for 16 bit conversion (Computed as 0/65535, 32768/65535, 32768/65535) */
 	RENDERCORE_API extern const FVector YUVOffsetNoScale16bits;
 
+	/** YUV Offset for float conversion (Computed as 1/16, 1/2, 1/2) */
+	RENDERCORE_API extern const FVector YUVOffsetFloat;
+
+	/** YUV Offset for float conversion (Computed as 0, 1/2, 1/2) */
+	RENDERCORE_API extern const FVector YUVOffsetNoScaleFloat;
+
 	/** Combine color transform matrix with yuv offset in a single matrix */
 	RENDERCORE_API FMatrix CombineColorTransformAndOffset(const FMatrix& InMatrix, const FVector& InYUVOffset);
 }
@@ -222,7 +228,7 @@ public:
 		: FGlobalShader(Initializer)
 	{ }
 
-	RENDERCORE_API void SetParameters(FRHIBatchedShaderParameters& BatchedParameters, const FIntPoint & TexDim, FShaderResourceViewRHIRef SRV_Y, FShaderResourceViewRHIRef SRV_UV, const FIntPoint& OutputDimensions, const FMatrix& ColorTransform, const FVector& YUVOffset, bool SrgbToLinear);
+	RENDERCORE_API void SetParameters(FRHIBatchedShaderParameters& BatchedParameters, const FIntPoint & TexDim, FShaderResourceViewRHIRef SRV_Y, FShaderResourceViewRHIRef SRV_UV, const FIntPoint& OutputDimensions, const FMatrix44f& ColorTransform, bool SrgbToLinear, const FMatrix44f& CSTransform, bool bSwapChroma);
 };
 
 
@@ -246,7 +252,7 @@ public:
 		: FGlobalShader(Initializer)
 	{ }
 
-	RENDERCORE_API void SetParameters(FRHIBatchedShaderParameters& BatchedParameters, TRefCountPtr<FRHITexture2D> NV12Texture, const FIntPoint& OutputDimensions, const FMatrix& ColorTransform, const FVector& YUVOffset, bool SrgbToLinear);
+	RENDERCORE_API void SetParameters(FRHIBatchedShaderParameters& BatchedParameters, TRefCountPtr<FRHITexture2D> NV12Texture, const FIntPoint& OutputDimensions, const FMatrix44f& ColorTransform, bool SrgbToLinear, const FMatrix44f& CSTransform, bool bSwapChroma);
 };
 
 
@@ -288,7 +294,7 @@ public:
 		: FGlobalShader(Initializer)
 	{ }
 
-	RENDERCORE_API void SetParameters(FRHIBatchedShaderParameters& BatchedParameters, const FIntPoint& TexDim, FShaderResourceViewRHIRef SRV_Y, FShaderResourceViewRHIRef SRV_UV, const FIntPoint& OutputDimensions, const FMatrix44f& ColorTransform, const FMatrix44f& CSTransform, bool bIsST2084);
+	RENDERCORE_API void SetParameters(FRHIBatchedShaderParameters& BatchedParameters, const FIntPoint& TexDim, FShaderResourceViewRHIRef SRV_Y, FShaderResourceViewRHIRef SRV_UV, const FIntPoint& OutputDimensions, const FMatrix44f& ColorTransform, const FMatrix44f& CSTransform, bool SrgbToLinear, bool bIsST2084);
 };
 
 /**
@@ -307,7 +313,7 @@ public:
 		: FGlobalShader(Initializer)
 	{ }
 
-	RENDERCORE_API void SetParameters(FRHIBatchedShaderParameters& BatchedParameters, const FIntPoint& TexDim, TRefCountPtr<FRHITexture2D> NV12Texture, const FIntPoint& OutputDimensions, const FMatrix44f& ColorTransform, const FMatrix44f& CSTransform, bool bIsST2084);
+	RENDERCORE_API void SetParameters(FRHIBatchedShaderParameters& BatchedParameters, const FIntPoint& TexDim, TRefCountPtr<FRHITexture2D> NV12Texture, const FIntPoint& OutputDimensions, const FMatrix44f& ColorTransform, const FMatrix44f& CSTransform, bool SrgbToLinear, bool bIsST2084);
 };
 
 /**
@@ -326,7 +332,7 @@ public:
 		: FGlobalShader(Initializer)
 	{ }
 
-	RENDERCORE_API void SetParameters(FRHIBatchedShaderParameters& BatchedParameters, const FIntPoint& TexDim, FShaderResourceViewRHIRef SRV_Y, FShaderResourceViewRHIRef SRV_U, FShaderResourceViewRHIRef SRV_V, const FIntPoint& OutputDimensions, const FMatrix44f& ColorTransform, const FMatrix44f& CSTransform, bool bIsST2084);
+	RENDERCORE_API void SetParameters(FRHIBatchedShaderParameters& BatchedParameters, const FIntPoint& TexDim, FShaderResourceViewRHIRef SRV_Y, FShaderResourceViewRHIRef SRV_U, FShaderResourceViewRHIRef SRV_V, const FIntPoint& OutputDimensions, const FMatrix44f& ColorTransform, const FMatrix44f& CSTransform, bool SrgbToLinear, bool bIsST2084);
 };
 
 /**
