Ensuring that MediaTetxure instance exclusively used with GUI rendering in Slate will use sRGB as their target color space even if the WCS is setup otherwise
#rnx

diff --git a/Engine/Plugins/Media/MediaPlayerEditor/Source/MediaPlayerEditor/Private/Toolkits/MediaSourceEditorToolkit.cpp b/Engine/Plugins/Media/MediaPlayerEditor/Source/MediaPlayerEditor/Private/Toolkits/MediaSourceEditorToolkit.cpp
--- a/Engine/Plugins/Media/MediaPlayerEditor/Source/MediaPlayerEditor/Private/Toolkits/MediaSourceEditorToolkit.cpp
+++ b/Engine/Plugins/Media/MediaPlayerEditor/Source/MediaPlayerEditor/Private/Toolkits/MediaSourceEditorToolkit.cpp
@@ -86,6 +86,7 @@ void FMediaSourceEditorToolkit::Initialize(UMediaSource* InMediaSource, const ET
 	{
 		MediaTexture->AutoClear = true;
 		MediaTexture->SetMediaPlayer(MediaPlayer);
+		MediaTexture->SetColorSpaceOverride(UE::Color::EColorSpace::sRGB);
 		MediaTexture->UpdateResource();
 	}
 
diff --git a/Engine/Plugins/Media/MediaPlayerEditor/Source/MediaPlayerEditor/Private/Widgets/SMediaPlayerEditorOutput.cpp b/Engine/Plugins/Media/MediaPlayerEditor/Source/MediaPlayerEditor/Private/Widgets/SMediaPlayerEditorOutput.cpp
--- a/Engine/Plugins/Media/MediaPlayerEditor/Source/MediaPlayerEditor/Private/Widgets/SMediaPlayerEditorOutput.cpp
+++ b/Engine/Plugins/Media/MediaPlayerEditor/Source/MediaPlayerEditor/Private/Widgets/SMediaPlayerEditorOutput.cpp
@@ -96,6 +96,7 @@ void SMediaPlayerEditorOutput::Construct(const FArguments& InArgs, UMediaPlayer&
 		{
 			MediaTexture->AutoClear = true;
 			MediaTexture->SetMediaPlayer(&InMediaPlayer);
+			MediaTexture->SetColorSpaceOverride(UE::Color::EColorSpace::sRGB);
 			MediaTexture->UpdateResource();
 			MediaTexture->AddToRoot();
 		}
diff --git a/Engine/Source/Runtime/MediaAssets/Private/Assets/MediaTexture.cpp b/Engine/Source/Runtime/MediaAssets/Private/Assets/MediaTexture.cpp
--- a/Engine/Source/Runtime/MediaAssets/Private/Assets/MediaTexture.cpp
+++ b/Engine/Source/Runtime/MediaAssets/Private/Assets/MediaTexture.cpp
@@ -88,6 +88,7 @@ UMediaTexture::UMediaTexture(const FObjectInitializer& ObjectInitializer)
 	, CachedNextSampleTime(FTimespan::MinValue())
 	, TextureNumMips(1)
 	, MipMapBias(0.0f)
+	, ColorspaceOverride(UE::Color::EColorSpace::None)
 {
 	NeverStream = true;
 	SRGB = true;
@@ -174,7 +175,7 @@ FTextureResource* UMediaTexture::CreateResource()
 		Filter = (TextureNumMips > 1) ? TF_Trilinear : TF_Bilinear;
 	}
 
-	return new FMediaTextureResource(*this, Dimensions, Size, ClearColor, CurrentGuid.IsValid() ? CurrentGuid : DefaultGuid, EnableGenMips, NumMips);
+	return new FMediaTextureResource(*this, Dimensions, Size, ClearColor, CurrentGuid.IsValid() ? CurrentGuid : DefaultGuid, EnableGenMips, NumMips, ColorspaceOverride);
 }
 
 
diff --git a/Engine/Source/Runtime/MediaAssets/Private/Misc/MediaTextureResource.cpp b/Engine/Source/Runtime/MediaAssets/Private/Misc/MediaTextureResource.cpp
--- a/Engine/Source/Runtime/MediaAssets/Private/Misc/MediaTextureResource.cpp
+++ b/Engine/Source/Runtime/MediaAssets/Private/Misc/MediaTextureResource.cpp
@@ -26,7 +26,6 @@
 #include "ProfilingDebugging/CsvProfiler.h"
 #include "Async/Async.h"
 #include "RenderGraphUtils.h"
-#include "ColorSpace.h"
 #include "HDRHelper.h"
 
 #include "MediaTexture.h"
@@ -321,7 +320,7 @@ namespace MediaTextureResourceHelpers
 /* FMediaTextureResource structors
  *****************************************************************************/
 
-FMediaTextureResource::FMediaTextureResource(UMediaTexture& InOwner, FIntPoint& InOwnerDim, SIZE_T& InOwnerSize, FLinearColor InClearColor, FGuid InTextureGuid, bool InEnableGenMips, uint8 InNumMips)
+FMediaTextureResource::FMediaTextureResource(UMediaTexture& InOwner, FIntPoint& InOwnerDim, SIZE_T& InOwnerSize, FLinearColor InClearColor, FGuid InTextureGuid, bool InEnableGenMips, uint8 InNumMips, UE::Color::EColorSpace OverrideColorSpaceType)
 	: Cleared(false)
 	, CurrentClearColor(InClearColor)
 	, InitialTextureGuid(InTextureGuid)
@@ -339,6 +338,11 @@ FMediaTextureResource::FMediaTextureResource(UMediaTexture& InOwner, FIntPoint&
 #else
 	bUsesImageExternal = !Owner.NewStyleOutput && GSupportsImageExternal;
 #endif
+	if (OverrideColorSpaceType != UE::Color::EColorSpace::None)
+	{
+		OverrideColorSpace.Reset(new UE::Color::FColorSpace(OverrideColorSpaceType));
+		check(OverrideColorSpace.IsValid());
+	}
 }
 
 
@@ -939,7 +943,7 @@ void FMediaTextureResource::ConvertSample(const TSharedPtr<IMediaTextureSample,
 
 void FMediaTextureResource::GetColorSpaceConversionMatrixForSample(const TSharedPtr<IMediaTextureSample, ESPMode::ThreadSafe> Sample, FMatrix44f& ColorSpaceMtx)
 {
-	const UE::Color::FColorSpace& Working = UE::Color::FColorSpace::GetWorking();
+	const UE::Color::FColorSpace& Working = OverrideColorSpace.IsValid() ? *OverrideColorSpace : UE::Color::FColorSpace::GetWorking();
 	ColorSpaceMtx = UE::Color::Transpose<float>(Working.GetXYZToRgb()) * Sample->GetGamutToXYZMatrix();
 	float NF = Sample->GetHDRNitsNormalizationFactor();
 	if (NF != 1.0f)
diff --git a/Engine/Source/Runtime/MediaAssets/Private/Misc/MediaTextureResource.h b/Engine/Source/Runtime/MediaAssets/Private/Misc/MediaTextureResource.h
--- a/Engine/Source/Runtime/MediaAssets/Private/Misc/MediaTextureResource.h
+++ b/Engine/Source/Runtime/MediaAssets/Private/Misc/MediaTextureResource.h
@@ -17,6 +17,7 @@
 #include "Async/Async.h"
 #include "RenderingThread.h"
 #include "RendererInterface.h"
+#include "ColorSpace.h"
 
 class FMediaPlayerFacade;
 class IMediaPlayer;
@@ -48,7 +49,7 @@ public:
 	 * @param bEnableGenMips If true mips generation will be enabled (possibly optimizing for NumMips == 1 case)
 	 * @param InNumMips The initial number of mips to be generated for the output texture
 	 */
-	MEDIAASSETS_API FMediaTextureResource(UMediaTexture& InOwner, FIntPoint& InOwnerDim, SIZE_T& InOwnerSize, FLinearColor InClearColor, FGuid InTextureGuid, bool bEnableGenMips, uint8 InNumMips);
+	MEDIAASSETS_API FMediaTextureResource(UMediaTexture& InOwner, FIntPoint& InOwnerDim, SIZE_T& InOwnerSize, FLinearColor InClearColor, FGuid InTextureGuid, bool bEnableGenMips, uint8 InNumMips, UE::Color::EColorSpace OverrideColorSpaceType);
 
 	/** Virtual destructor. */
 	virtual ~FMediaTextureResource() 
@@ -254,4 +255,7 @@ private:
 	/** Cached FRenderParams, used when JustInTimeRender() gets called. */
 	TUniquePtr<FRenderParams> JustInTimeRenderParams;
 
+	/** Colorspace to override standard proejct "working color space' */
+	TUniquePtr<UE::Color::FColorSpace> OverrideColorSpace;
+
 };
diff --git a/Engine/Source/Runtime/MediaAssets/Public/MediaTexture.h b/Engine/Source/Runtime/MediaAssets/Public/MediaTexture.h
--- a/Engine/Source/Runtime/MediaAssets/Public/MediaTexture.h
+++ b/Engine/Source/Runtime/MediaAssets/Public/MediaTexture.h
@@ -220,6 +220,14 @@ public:
 	 */
 	void SetMipMapBias(float InMipMapBias);
 
+	/**
+	 * Set an override for color space conversions
+	 */
+	void SetColorSpaceOverride(UE::Color::EColorSpace InColorspaceOverride)
+	{
+		ColorspaceOverride = InColorspaceOverride;
+	}
+
 	/** Renders this media texture. Only has an effect if its RenderMode is ERenderMode::JustInTime */
 	virtual void JustInTimeRender();
 
@@ -342,4 +350,7 @@ private:
 
 	/** Current render mode of this media texture. Can be changed using SetRenderMode() */
 	ERenderMode RenderMode = ERenderMode::Default;
+
+	/** Colorspace override */
+	UE::Color::EColorSpace ColorspaceOverride;
 };
