Adding explict checks for CS setuip differences to trigger conversion of sample data ot MFW
#rnx

diff --git a/Engine/Source/Runtime/MediaAssets/Private/Misc/MediaTextureResource.cpp b/Engine/Source/Runtime/MediaAssets/Private/Misc/MediaTextureResource.cpp
--- a/Engine/Source/Runtime/MediaAssets/Private/Misc/MediaTextureResource.cpp
+++ b/Engine/Source/Runtime/MediaAssets/Private/Misc/MediaTextureResource.cpp
@@ -835,23 +835,50 @@ bool FMediaTextureResource::RequiresConversion(const TSharedPtr<IMediaTextureSam
 
 	// If the output dimensions are not the same as the sample's
 	// dimensions, a resizing conversion on the GPU is required.
-
+	//
+	// Some reasons this could trigger:
+	//	- Cropping
+	//  - NV12's 1.5x heigth factor -> but the format would need conversion to RGB anyway
+	//	- Various YCbCr formats will be stored with different widths as the final RGB output data -> a conversion is needed anyway
+	//
 	if (Sample->GetDim() != Sample->GetOutputDim())
 	{
 		return true;
 	}
 
-//TODO: CS Changes should trigger conversion in any case
+	// Color space different?
+	const UE::Color::FColorSpace& Working = OverrideColorSpace.IsValid() ? *OverrideColorSpace : UE::Color::FColorSpace::GetWorking();
+	const float Tollerance = 1.e-7f;
+	if (!Sample->GetDisplayPrimaryRed().Equals(FVector2f(Working.GetRedChromaticity()), Tollerance) ||
+		!Sample->GetDisplayPrimaryGreen().Equals(FVector2f(Working.GetGreenChromaticity()), Tollerance) ||
+		!Sample->GetDisplayPrimaryBlue().Equals(FVector2f(Working.GetBlueChromaticity()), Tollerance) ||
+		!Sample->GetWhitePoint().Equals(FVector2f(Working.GetWhiteChromaticity()), Tollerance))
+	{
+		// Yes! We need to convert...
+		return true;
+	}
 
+	// No. Lets check if the encoding type is linear or not...
 	UE::Color::EEncoding ColorEncoding = Sample->GetEncodingType();
 	EMediaTextureSampleFormat Format = Sample->GetFormat();
 
-	bool b8BitRGBA = (Format == EMediaTextureSampleFormat::CharBGRA || Format == EMediaTextureSampleFormat::CharRGBA);
+	bool b8BitRGBA = (Format == EMediaTextureSampleFormat::CharBGRA ||
+					  Format == EMediaTextureSampleFormat::CharRGBA ||
+					  Format == EMediaTextureSampleFormat::DXT1 ||
+					  Format == EMediaTextureSampleFormat::DXT5);
 
-	// 8-Bit RGBA will need no processing if its either linear or sRGB (we configure the texture objects as sRGB an get the conversion for free)
-	if (b8BitRGBA && (ColorEncoding == UE::Color::EEncoding::sRGB || ColorEncoding == UE::Color::EEncoding::Linear))
+	if (b8BitRGBA)
 	{
-		return false;
+		// sRGB?
+		if (ColorEncoding == UE::Color::EEncoding::sRGB)
+		{
+			// 8-bit RGBA and either CPU side buffer or texture with sRGB flag -> we can assume an inpout texture with HW translation to linear
+			auto Texture = Sample->GetTexture();
+			if (Texture == nullptr || ((uint32)Texture->GetFlags() & (uint32)TexCreate_SRGB))
+			{
+				return false;
+			}
+		}
 	}
 
 	bool bRGBA = b8BitRGBA ||
@@ -861,14 +888,14 @@ bool FMediaTextureResource::RequiresConversion(const TSharedPtr<IMediaTextureSam
 				 Format == EMediaTextureSampleFormat::RGBA16 ||
 				 Format == EMediaTextureSampleFormat::CharBGR10A2;
 
-	// If we got RGBA (not caught by the special case above) with an encoding other than linear we need to convert...
-	if (bRGBA && ColorEncoding != UE::Color::EEncoding::Linear)
+	// RGBA and linear?
+	if (bRGBA && ColorEncoding == UE::Color::EEncoding::Linear)
 	{
-		return true;
+		return false;
 	}
 
-	// And if we get here: any RGBA will be able to go out unchanged!
-	return !bRGBA;
+	// Either we have RGBA with non-linear encoding or its not RGBA -> we must convert!
+	return true;
 	}
 
 
