Introducing a meta sound MediaStreamer node
#rnx

diff --git a/Engine/Source/Runtime/MediaAssets/MediaAssets.Build.cs b/Engine/Source/Runtime/MediaAssets/MediaAssets.Build.cs
--- a/Engine/Source/Runtime/MediaAssets/MediaAssets.Build.cs
+++ b/Engine/Source/Runtime/MediaAssets/MediaAssets.Build.cs
@@ -10,6 +10,7 @@ namespace UnrealBuildTool.Rules
 				new string[] {
 					"SignalProcessing",
 					"AudioMixer",
+					"AudioExtensions",
                     "Core",
 					"CoreUObject",
 					"Engine",
diff --git a/Engine/Source/Runtime/MediaAssets/Private/Assets/MediaPlayer.cpp b/Engine/Source/Runtime/MediaAssets/Private/Assets/MediaPlayer.cpp
--- a/Engine/Source/Runtime/MediaAssets/Private/Assets/MediaPlayer.cpp
+++ b/Engine/Source/Runtime/MediaAssets/Private/Assets/MediaPlayer.cpp
@@ -956,6 +956,19 @@ void UMediaPlayer::PostEditChangeProperty(FPropertyChangedEvent& PropertyChanged
 
 #endif
 
+/* IAudioProxyDataFactory overrides
+ *****************************************************************************/
+
+TSharedPtr<Audio::IProxyData> UMediaPlayer::CreateProxyData(const Audio::FProxyDataInitParams& InitParams)
+{
+//WE GET HERE WHEN STARTING THE GRAPH CREATION ON PLAY (well, at least with the editor) - GAMETHREAD
+	if (!Proxy.IsValid())
+	{
+		Proxy = MakeShared<FMediaPlayerProxy, ESPMode::ThreadSafe>(this);
+	}
+	return Proxy;
+}
+
 
 /* UMediaPlayer callbacks
  *****************************************************************************/
@@ -1248,4 +1261,14 @@ void UMediaPlayer::OpenSourceLatent(const UObject* WorldContextObject, FLatentAc
 	}
 }
 
+// --------------------------------------------------------------------------------------------------------------------------------------------------------------
+// --------------------------------------------------------------------------------------------------------------------------------------------------------------
+
+FMediaPlayerProxy::FMediaPlayerProxy(UMediaPlayer* Player)
+{
+	PlayerFacade = Player->GetPlayerFacade();
+}
 
+FMediaPlayerProxy::~FMediaPlayerProxy()
+{
+}
diff --git a/Engine/Source/Runtime/MediaAssets/Public/MediaPlayer.h b/Engine/Source/Runtime/MediaAssets/Public/MediaPlayer.h
--- a/Engine/Source/Runtime/MediaAssets/Public/MediaPlayer.h
+++ b/Engine/Source/Runtime/MediaAssets/Public/MediaPlayer.h
@@ -17,6 +17,7 @@
 #include "Engine/LatentActionManager.h"
 #include "MediaPlayerOptions.h"
 #include "IMediaTimeSource.h"
+#include "IAudioProxyInitializer.h"
 
 #include "MediaPlayer.generated.h"
 
@@ -26,6 +27,7 @@ class IMediaModule;
 class IMediaMetadataItem;
 class UMediaPlaylist;
 class UMediaSource;
+class UMediaPlayer;
 
 enum class EMediaEvent;
 
@@ -88,12 +90,38 @@ public:
 };
 
 
+/**
+ * Proxy implementation for Metasound integration
+ */
+class MEDIAASSETS_API FMediaPlayerProxy : public Audio::TProxyData<FMediaPlayerProxy>
+{
+public:
+	IMPL_AUDIOPROXY_CLASS(FMediaPlayerProxy);
+
+	explicit FMediaPlayerProxy(UMediaPlayer* Player);
+
+	FMediaPlayerProxy(const FMediaPlayerProxy& Other) = default;
+
+	~FMediaPlayerProxy();
+
+	TSharedPtr<FMediaPlayerFacade, ESPMode::ThreadSafe> GetPlayerFacade() const
+	{
+		return(PlayerFacade.Pin());
+	}
+
+private:
+	TWeakPtr<FMediaPlayerFacade, ESPMode::ThreadSafe> PlayerFacade;
+};
+using FMediaPlayerProxyPtr = TSharedPtr<FMediaPlayerProxy, ESPMode::ThreadSafe>;
+
+
 /**
  * Implements a media player asset that can play movies and other media sources.
  */
 UCLASS(BlueprintType, hidecategories=(Object))
 class MEDIAASSETS_API UMediaPlayer
 	: public UObject
+	, public IAudioProxyDataFactory
 {
 	GENERATED_UCLASS_BODY()
 
@@ -1181,6 +1209,13 @@ public:
 	UPROPERTY(BlueprintReadWrite, EditAnywhere, Category=Playback)
 	uint32 Shuffle : 1;
 
+public:
+	//
+	// IAudioProxyDataFactory interface (for Meta Sound support)
+	//
+
+	virtual TSharedPtr<Audio::IProxyData> CreateProxyData(const Audio::FProxyDataInitParams& InitParams) override;
+
 protected:
 
 	/**
@@ -1297,6 +1332,9 @@ private:
 	/** Has registered with the media module */
 	bool RegisteredWithMediaModule;
 
+	/** Proxy instance for integration with Metasound */
+	FMediaPlayerProxyPtr Proxy;
+
 #if WITH_EDITORONLY_DATA
 public:
 	/** Whether this player should stop when entering or exiting PIE. */
