Adding blueprint accessible struct and function to MediaPlayer for media metadata

diff --git a/Engine/Source/Runtime/MediaAssets/Private/Assets/MediaPlayer.cpp b/Engine/Source/Runtime/MediaAssets/Private/Assets/MediaPlayer.cpp
--- a/Engine/Source/Runtime/MediaAssets/Private/Assets/MediaPlayer.cpp
+++ b/Engine/Source/Runtime/MediaAssets/Private/Assets/MediaPlayer.cpp
@@ -233,6 +233,104 @@ TSharedPtr<TMap<FString, TArray<TUniquePtr<IMediaMetadataItem>>>, ESPMode::Threa
 	return PlayerFacade->GetMediaMetadata();
 }
 
+TMap<FString, FMediaMetadataItemsBPT> UMediaPlayer::GetMediaMetadataItems() const
+{
+	TMap<FString, FMediaMetadataItemsBPT> OutMetadata;
+	if(const TMap<FString, TArray<TUniquePtr<IMediaMetadataItem>>>* MetadataPtr = GetMediaMetadata().Get())
+	{
+		for(const TPair<FString, TArray<TUniquePtr<IMediaMetadataItem>>>& Data : *MetadataPtr)
+		{
+			FMediaMetadataItemsBPT OutItems;
+			for(const TUniquePtr<IMediaMetadataItem>& Item : Data.Value)
+			{
+				if(const IMediaMetadataItem* ItemPtr = Item.Get())
+				{
+					FMediaMetadataItemBPT OutItem;
+					OutItem.LanguageCode = ItemPtr->GetLanguageCode();
+					OutItem.MimeType = ItemPtr->GetMimeType();
+					FVariant TempItem = ItemPtr->GetValue();
+			
+					switch (TempItem.GetType())
+					{
+					case EVariantTypes::String:
+						{
+							OutItem.StringData = TempItem.GetValue<FString>();
+							break;
+						}
+					case EVariantTypes::Ansichar:
+					case EVariantTypes::ByteArray:
+						{
+							OutItem.BinaryData = TempItem.GetValue<TArray<uint8>>();
+							break;
+						}
+			
+					case EVariantTypes::Bool:
+						{
+							OutItem.StringData = TempItem.GetValue<bool>() ? TEXT("true") : TEXT("false");
+							break;
+						}
+					case EVariantTypes::Double:
+						{
+							OutItem.StringData = FString::SanitizeFloat(TempItem.GetValue<double>());
+							break;
+						}
+					case EVariantTypes::Float:
+						{
+							OutItem.StringData = FString::SanitizeFloat(TempItem.GetValue<float>());
+							break;
+						}
+					case EVariantTypes::UInt8:
+						{
+							OutItem.StringData = FString::Printf(TEXT("%d"), TempItem.GetValue<uint8>());
+							break;
+						}
+					case EVariantTypes::Int8:
+						{
+							OutItem.StringData = FString::Printf(TEXT("%d"), TempItem.GetValue<int8>());
+							break;
+						}
+					case EVariantTypes::UInt16:
+						{
+							OutItem.StringData = FString::Printf(TEXT("%d"), TempItem.GetValue<uint16>());
+							break;
+						}
+					case EVariantTypes::Int16:
+						{
+							OutItem.StringData = FString::Printf(TEXT("%d"), TempItem.GetValue<int16>());
+							break;
+						}
+					case EVariantTypes::UInt32:
+						{
+							OutItem.StringData = FString::Printf(TEXT("%d"), TempItem.GetValue<int32>());
+							break;
+						}
+					case EVariantTypes::Int32:
+						{
+							OutItem.StringData = FString::Printf(TEXT("%d"), TempItem.GetValue<uint32>());
+							break;
+						}
+					case EVariantTypes::UInt64:
+						{
+							OutItem.StringData = FString::Printf(TEXT("%llu"), TempItem.GetValue<uint64>());
+							break;
+						}
+					case EVariantTypes::Int64:
+						{
+							OutItem.StringData = FString::Printf(TEXT("%lld"), TempItem.GetValue<int64>());
+							break;
+						}
+					default:
+						UE_LOG(LogMediaAssets, Warning, TEXT("GetMediaMetadataItems() unhandled FVariant type!"));
+					}
+					OutItems.Items.Add(OutItem);
+				}
+			}
+			OutMetadata.Add(Data.Key, OutItems);
+		}
+	}
+	return OutMetadata;
+}
+
 TSharedRef<FMediaPlayerFacade, ESPMode::ThreadSafe> UMediaPlayer::GetPlayerFacade() const
 {
 	return PlayerFacade.ToSharedRef();
diff --git a/Engine/Source/Runtime/MediaAssets/Public/MediaPlayer.h b/Engine/Source/Runtime/MediaAssets/Public/MediaPlayer.h
--- a/Engine/Source/Runtime/MediaAssets/Public/MediaPlayer.h
+++ b/Engine/Source/Runtime/MediaAssets/Public/MediaPlayer.h
@@ -6,6 +6,7 @@
 #include "Containers/UnrealString.h"
 #include "Delegates/Delegate.h"
 #include "IMediaOptions.h"
+#include "IMediaMetadataItem.h"
 #include "Math/Quat.h"
 #include "Math/Rotator.h"
 #include "Templates/SharedPointer.h"
@@ -89,6 +90,33 @@ public:
 	int64 SequenceIndex;
 };
 
+USTRUCT(BlueprintType)
+struct MEDIAASSETS_API FMediaMetadataItemBPT
+{
+	GENERATED_BODY()
+
+public:
+	UPROPERTY(BlueprintReadOnly, Category = "Media|Metadata")
+	FString LanguageCode;
+	
+	UPROPERTY(BlueprintReadOnly, Category = "Media|Metadata")
+	FString MimeType;
+	
+	UPROPERTY(BlueprintReadOnly, Category = "Media|Metadata")
+	FString StringData;
+	
+	UPROPERTY(BlueprintReadOnly, Category = "Media|Metadata")
+	TArray<uint8> BinaryData;
+};
+
+USTRUCT(BlueprintType)
+struct MEDIAASSETS_API FMediaMetadataItemsBPT
+{
+	GENERATED_BODY()
+public:
+	UPROPERTY(BlueprintReadOnly, Category = "Media|Metadata")
+	TArray<FMediaMetadataItemBPT> Items;
+};
 
 /**
  * Proxy implementation for Metasound integration
@@ -1068,6 +1096,14 @@ public:
 	 */
 	TSharedPtr<TMap<FString, TArray<TUniquePtr<IMediaMetadataItem>>>, ESPMode::ThreadSafe> GetMediaMetadata() const;
 
+	/**
+	 * This is the blueprint accessible version of the GetMediaMetadata.
+	 * @return Map with arrays of FMediaMetaDataItem entries describing any metadata found in the current stream
+	 * @note Listen to EMediaEvent::MetadataChanged to catch updates to this data
+	 */
+	UFUNCTION(BlueprintCallable, Category = "Media|Metadata")
+	TMap<FString, FMediaMetadataItemsBPT> GetMediaMetadataItems() const;
+	
 	/**
 	 * Get the media player facade that manages low-level media players
 	 *
