Fixing MFW conversion code for DXT5 textures (HAP)
#rnx

diff --git a/Engine/Source/Runtime/MediaAssets/Private/Misc/MediaTextureResource.cpp b/Engine/Source/Runtime/MediaAssets/Private/Misc/MediaTextureResource.cpp
--- a/Engine/Source/Runtime/MediaAssets/Private/Misc/MediaTextureResource.cpp
+++ b/Engine/Source/Runtime/MediaAssets/Private/Misc/MediaTextureResource.cpp
@@ -1491,6 +1491,13 @@ void FMediaTextureResource::CreateOutputRenderTarget(const FIntPoint & InDim, EP
 		InNumMips = FMath::Min(InNumMips, MaxMips);
 	}
 
+	// Avoid setting things up as a render target if we are passing in a block compression format
+	// (this happens only if these have no conversion needs at all - otherwise they will be promoted to RGBA8)
+	if (InPixelFormat != PF_DXT1 && InPixelFormat != PF_DXT5 && InPixelFormat != PF_BC4)
+	{
+		OutputCreateFlags |= ETextureCreateFlags::RenderTargetable;
+	}
+
 	if ((InClearColor != CurrentClearColor) || !OutputTarget.IsValid() || (OutputTarget->GetSizeXY() != InDim) || (OutputTarget->GetFormat() != InPixelFormat) || ((OutputTarget->GetFlags() & OutputCreateFlags) != OutputCreateFlags) || CurrentNumMips != InNumMips)
 	{
 		MipGenerationCache.SafeRelease();
@@ -1501,7 +1508,7 @@ void FMediaTextureResource::CreateOutputRenderTarget(const FIntPoint & InDim, EP
 			.SetExtent(InDim)
 			.SetFormat(InPixelFormat)
 			.SetNumMips(InNumMips)
-			.SetFlags(OutputCreateFlags | ETextureCreateFlags::RenderTargetable | ETextureCreateFlags::ShaderResource)
+			.SetFlags(OutputCreateFlags | ETextureCreateFlags::ShaderResource)
 			.SetInitialState(ERHIAccess::SRVMask)
 			.SetClearValue(FClearValueBinding(InClearColor))
 			.SetClassName(ClassName)
