Add custom color converter support to media texture sample

diff --git a/Engine/Source/Runtime/Media/Public/IMediaTextureSample.h b/Engine/Source/Runtime/Media/Public/IMediaTextureSample.h
--- a/Engine/Source/Runtime/Media/Public/IMediaTextureSample.h
+++ b/Engine/Source/Runtime/Media/Public/IMediaTextureSample.h
@@ -22,6 +22,7 @@
 #if WITH_ENGINE
 	class FRHITexture;
 	class IMediaTextureSampleConverter;
+	class IMediaTextureSampleColorConverter;
 #endif
 
 
@@ -266,6 +267,16 @@ public:
 		return nullptr; 
 	}
 
+	/**
+	 * Get a media texture sample color converter if sample implements it
+	 * @Note IMediaTextureSampleColorConverter will be applied after IMediaTextureSampleConverter if one is provided.
+	 * @return texture sample color converter
+	 */
+	virtual IMediaTextureSampleColorConverter* GetMediaTextureSampleColorConverter()
+	{ 
+		return nullptr; 
+	}
+
 #endif //WITH_ENGINE
 
 	/**
diff --git a/Engine/Source/Runtime/Media/Public/IMediaTextureSampleConverter.h b/Engine/Source/Runtime/Media/Public/IMediaTextureSampleConverter.h
--- a/Engine/Source/Runtime/Media/Public/IMediaTextureSampleConverter.h
+++ b/Engine/Source/Runtime/Media/Public/IMediaTextureSampleConverter.h
@@ -35,3 +35,19 @@ public:
 
 	virtual bool Convert(FTexture2DRHIRef & InDstTexture, const FConversionHints & Hints) = 0;
 };
+
+/**
+ * Interface class to implement custom sample color conversion
+ */
+class IMediaTextureSampleColorConverter
+{
+public:
+	virtual ~IMediaTextureSampleColorConverter() {}
+
+	/**
+	 * Apply a color conversion on the input and store the result in the destination texture.
+	 * @return true If the color conversion was successfully applied.
+	 */
+	virtual bool ApplyColorConversion(FTexture2DRHIRef& InSrcTexture, FTexture2DRHIRef& InDstTexture) = 0;
+};
+
diff --git a/Engine/Source/Runtime/MediaAssets/Private/Misc/MediaTextureResource.cpp b/Engine/Source/Runtime/MediaAssets/Private/Misc/MediaTextureResource.cpp
--- a/Engine/Source/Runtime/MediaAssets/Private/Misc/MediaTextureResource.cpp
+++ b/Engine/Source/Runtime/MediaAssets/Private/Misc/MediaTextureResource.cpp
@@ -17,6 +17,7 @@
 #include "MediaSampleSource.h"
 #include "MediaShaders.h"
 #include "PipelineStateCache.h"
+#include "RenderGraphUtils.h"
 #include "SceneUtils.h"
 #include "Shader.h"
 #include "StaticBoundShaderState.h"
@@ -25,7 +26,8 @@
 #include "GenerateMips.h"
 #include "ProfilingDebugging/CsvProfiler.h"
 #include "Async/Async.h"
-#include "RenderGraphUtils.h"
+
+
 #include "HDRHelper.h"
 
 #include "MediaTexture.h"
@@ -428,7 +430,7 @@ void FMediaTextureResource::Render(const FRenderParams& Params)
 #endif
 		// - - - - - - - - - - - - - - - - - - - - - - - - - - - - - - - - - - - -
 		// - - - - - - - - - - - - - - - - - - - - - - - - - - - - - - - - - - - -
-		
+
 		const uint8 NumMips = bEnableGenMips ? Params.NumMips : 1;
 
 		// If real "external texture" support is in place and no mips are used the image will "bypass" any of this processing via the GUID-based lookup for "ExternalTextures" and will
@@ -480,7 +482,7 @@ void FMediaTextureResource::Render(const FRenderParams& Params)
 					{
 						// Conversion is fully handled by converter
 
-						CreateOutputRenderTarget(Sample->GetOutputDim(), MediaTextureResourceHelpers::GetConvertedPixelFormat(Sample), MediaTextureResourceHelpers::RequiresSrgbTexture(Sample), Params.ClearColor, Hints.NumMips, bNeedsUAVTexture);
+						CreateIntermediateRenderTarget(Sample->GetOutputDim(), MediaTextureResourceHelpers::GetConvertedPixelFormat(Sample), MediaTextureResourceHelpers::RequiresSrgbTexture(Sample), Params.ClearColor, Hints.NumMips, bNeedsUAVTexture);
 						Converter->Convert(RenderTargetTextureRHI, Hints);
 					}
 				}
@@ -493,7 +495,7 @@ void FMediaTextureResource::Render(const FRenderParams& Params)
 						// As the converter created the texture, we might need to convert it even more to make it fit our needs. Check...
 						if (RequiresConversion(OutTexture, Sample->GetOutputDim(), NumMips))
 						{
-							CreateOutputRenderTarget(Sample->GetOutputDim(), MediaTextureResourceHelpers::GetConvertedPixelFormat(Sample), MediaTextureResourceHelpers::RequiresSrgbTexture(Sample), Params.ClearColor, Hints.NumMips, bNeedsUAVTexture);
+							CreateIntermediateRenderTarget(Sample->GetOutputDim(), MediaTextureResourceHelpers::GetConvertedPixelFormat(Sample), MediaTextureResourceHelpers::RequiresSrgbTexture(Sample), Params.ClearColor, Hints.NumMips, bNeedsUAVTexture);
 							ConvertTextureToOutput(OutTexture.GetReference(), Sample);
 						}
 						else
@@ -529,6 +531,45 @@ void FMediaTextureResource::Render(const FRenderParams& Params)
 				}
 			}
 
+			if (IMediaTextureSampleColorConverter* Converter = Sample->GetMediaTextureSampleColorConverter())
+			{
+				if (bRecreateOutputTarget)
+				{
+					const static FLazyName ClassName(TEXT("FMediaTextureResource"));
+
+					const FRHITextureDesc& IntermediateTextureDesc = IntermediateTarget->GetDesc();
+			
+					const FRHITextureCreateDesc Desc =
+						FRHITextureCreateDesc::Create2D(TEXT("MediaTextureColorConversionOutput"),
+							IntermediateTextureDesc.GetSize().X,
+							IntermediateTextureDesc.GetSize().Y,
+							IntermediateTextureDesc.Format)
+						.SetNumMips(IntermediateTextureDesc.NumMips)
+						.SetClearValue(IntermediateTextureDesc.ClearValue)
+						.SetClassName(ClassName)
+						.SetOwnerName(GetOwnerName())
+						.SetFlags(IntermediateTextureDesc.Flags);
+			
+					OutputTarget = RHICreateTexture(Desc);
+					OutputTarget->SetName(TEXT("MediaTextureResourceOutput"));
+					OutputTarget->SetOwnerName(GetOwnerName());
+			
+					bRecreateOutputTarget = false;
+				}
+		
+				Converter->ApplyColorConversion(IntermediateTarget, OutputTarget);
+
+				if (RenderTargetTextureRHI != OutputTarget)
+				{
+					UpdateTextureReference(OutputTarget);
+				}
+			}
+			else
+			{
+				OutputTarget = IntermediateTarget;
+				UpdateTextureReference(OutputTarget);
+			}
+
 			Rotation = Sample->GetScaleRotation();
 			Offset = Sample->GetOffset();
 
@@ -694,7 +735,7 @@ void FMediaTextureResource::InitRHI()
 	// Make sure init has done it's job - we can't leave here without valid bindable resources for some RHI's
 	check(TextureRHI.IsValid());
 	check(RenderTargetTextureRHI.IsValid());
-	check(OutputTarget.IsValid());
+	check(IntermediateTarget.IsValid());
 
 	// Register "external texture" parameters if the platform does not support them (and hence the player does not set them)
 	if (!bUsesImageExternal)
@@ -712,6 +753,7 @@ void FMediaTextureResource::ReleaseRHI()
 	MipGenerationCache.SafeRelease();
 
 	InputTarget.SafeRelease();
+	IntermediateTarget.SafeRelease();
 	OutputTarget.SafeRelease();
 	RenderTargetTextureRHI.SafeRelease();
 	TextureRHI.SafeRelease();
@@ -755,7 +797,7 @@ void FMediaTextureResource::ClearTexture(const FLinearColor& ClearColor, bool Sr
 	// create output render target if we don't have one yet
 	constexpr uint8 NumMips = 1;
 	constexpr bool bNeedsUAVTexture = false;
-	CreateOutputRenderTarget(FIntPoint(2, 2), PF_B8G8R8A8, SrgbOutput, ClearColor, NumMips, bNeedsUAVTexture);
+	CreateIntermediateRenderTarget(FIntPoint(2, 2), PF_B8G8R8A8, SrgbOutput, ClearColor, NumMips, bNeedsUAVTexture);
 
 	// draw the clear color
 	FRHICommandListImmediate& CommandList = FRHICommandListExecutor::GetImmediateCommandList();
@@ -896,7 +938,8 @@ bool FMediaTextureResource::RequiresConversion(const TSharedPtr<IMediaTextureSam
 
 	// Either we have RGBA with non-linear encoding or its not RGBA -> we must convert!
 	return true;
-	}
+}
+
 
 
 void FMediaTextureResource::ConvertSample(const TSharedPtr<IMediaTextureSample, ESPMode::ThreadSafe>& Sample, const FLinearColor& ClearColor, uint8 InNumMips)
@@ -962,7 +1005,7 @@ void FMediaTextureResource::ConvertSample(const TSharedPtr<IMediaTextureSample,
 	const FIntPoint OutputDim = Sample->GetOutputDim();
 	const uint8 NumMips = (SampleNumMips > 1) ? SampleNumMips : InNumMips;
 	const bool bNeedsUAVTexture = MediaTextureResourceHelpers::RequiresUAVTexture(Sample, NumMips);
-	CreateOutputRenderTarget(OutputDim, MediaTextureResourceHelpers::GetConvertedPixelFormat(Sample), MediaTextureResourceHelpers::RequiresSrgbTexture(Sample), ClearColor, NumMips, bNeedsUAVTexture);
+	CreateIntermediateRenderTarget(OutputDim, MediaTextureResourceHelpers::GetConvertedPixelFormat(Sample), MediaTextureResourceHelpers::RequiresSrgbTexture(Sample), ClearColor, NumMips, bNeedsUAVTexture);
 
 	ConvertTextureToOutput(InputTexture, Sample);
 }
@@ -971,7 +1014,16 @@ void FMediaTextureResource::ConvertSample(const TSharedPtr<IMediaTextureSample,
 void FMediaTextureResource::GetColorSpaceConversionMatrixForSample(const TSharedPtr<IMediaTextureSample, ESPMode::ThreadSafe> Sample, FMatrix44f& ColorSpaceMtx)
 {
 	const UE::Color::FColorSpace& Working = OverrideColorSpace.IsValid() ? *OverrideColorSpace : UE::Color::FColorSpace::GetWorking();
-	ColorSpaceMtx = UE::Color::Transpose<float>(Working.GetXYZToRgb()) * Sample->GetGamutToXYZMatrix();
+	
+	if (Sample->GetMediaTextureSampleColorConverter())
+	{
+		ColorSpaceMtx = FMatrix44f::Identity;
+	}
+	else
+	{
+		ColorSpaceMtx = UE::Color::Transpose<float>(Working.GetXYZToRgb()) * Sample->GetGamutToXYZMatrix();
+	}
+	
 	float NF = Sample->GetHDRNitsNormalizationFactor();
 	if (NF != 1.0f)
 	{
@@ -992,7 +1044,7 @@ void FMediaTextureResource::GetColorSpaceConversionMatrixForSample(const TShared
 		SCOPED_GPU_STAT(CommandList, MediaTextureResource);
 
 		FGraphicsPipelineStateInitializer GraphicsPSOInit;
-		FRHITexture* RenderTarget = RenderTargetTextureRHI.GetReference();
+		FRHITexture* RenderTarget = IntermediateTarget;
 		CommandList.Transition(FRHITransitionInfo(RenderTarget, ERHIAccess::Unknown, ERHIAccess::RTV));
 
 		FIntPoint OutputDim(RenderTarget->GetSizeXYZ().X, RenderTarget->GetSizeXYZ().Y);
@@ -1303,17 +1355,17 @@ void FMediaTextureResource::CopySample(const TSharedPtr<IMediaTextureSample, ESP
 			UpdateTextureReference(SampleTexture);
 
 			MipGenerationCache.SafeRelease();
-			OutputTarget.SafeRelease();
+			IntermediateTarget.SafeRelease();
 		}
 		else
 		{
 			// Texture to receive texture from sample
 			const uint8 NumMips = (SampleNumMips > 1) ? SampleNumMips : InNumMips;
 			const bool bNeedsUAVTexture = MediaTextureResourceHelpers::RequiresUAVTexture(Sample, NumMips);
-			CreateOutputRenderTarget(Sample->GetOutputDim(), MediaTextureResourceHelpers::GetPixelFormat(Sample), MediaTextureResourceHelpers::RequiresSrgbTexture(Sample), ClearColor, NumMips, bNeedsUAVTexture);
+			CreateIntermediateRenderTarget(Sample->GetOutputDim(), MediaTextureResourceHelpers::GetPixelFormat(Sample), MediaTextureResourceHelpers::RequiresSrgbTexture(Sample), ClearColor, NumMips, bNeedsUAVTexture);
 
 			// Copy data into the output texture to able to add mips later on
-			FRHICommandListExecutor::GetImmediateCommandList().CopyTexture(SampleTexture, OutputTarget, FRHICopyTextureInfo());
+			FRHICommandListExecutor::GetImmediateCommandList().CopyTexture(SampleTexture, IntermediateTarget, FRHICopyTextureInfo());
 		}
 	}
 	else
@@ -1321,7 +1373,7 @@ void FMediaTextureResource::CopySample(const TSharedPtr<IMediaTextureSample, ESP
 		// Texture to receive precisely only output pixels via CPU copy
 		const uint8 NumMips = (SampleNumMips > 1) ? SampleNumMips : InNumMips;
 		const bool bNeedsUAVTexture = MediaTextureResourceHelpers::RequiresUAVTexture(Sample, NumMips);
-		CreateOutputRenderTarget(Sample->GetDim(), MediaTextureResourceHelpers::GetPixelFormat(Sample), MediaTextureResourceHelpers::RequiresSrgbTexture(Sample), ClearColor, NumMips, bNeedsUAVTexture);
+		CreateIntermediateRenderTarget(Sample->GetDim(), MediaTextureResourceHelpers::GetPixelFormat(Sample), MediaTextureResourceHelpers::RequiresSrgbTexture(Sample), ClearColor, NumMips, bNeedsUAVTexture);
 
 		// If we also have no source buffer and the platform generally would allow for use of external textures, we assume it is just that...
 		// (as long as the player actually produces (dummy) samples, this will enable mips support as well as auto conversion for "new style output" mode)
@@ -1443,6 +1495,11 @@ void FMediaTextureResource::UpdateResourceSize()
 	{
 		ResourceSize += CalcTextureSize(InputTarget->GetSizeX(), InputTarget->GetSizeY(), InputTarget->GetFormat(), 1);
 	}
+	
+	if (IntermediateTarget.IsValid() && IntermediateTarget != OutputTarget)
+	{
+		ResourceSize += CalcTextureSize(IntermediateTarget->GetSizeX(), IntermediateTarget->GetSizeY(), IntermediateTarget->GetFormat(), 1);
+	}
 
 	if (OutputTarget.IsValid())
 	{
@@ -1472,7 +1529,7 @@ void FMediaTextureResource::UpdateTextureReference(FRHITexture* NewTexture)
 }
 
 
-void FMediaTextureResource::CreateOutputRenderTarget(const FIntPoint & InDim, EPixelFormat InPixelFormat, bool bInSRGB, const FLinearColor & InClearColor, uint8 InNumMips, bool bNeedsUAVSupport)
+void FMediaTextureResource::CreateIntermediateRenderTarget(const FIntPoint & InDim, EPixelFormat InPixelFormat, bool bInSRGB, const FLinearColor & InClearColor, uint8 InNumMips, bool bNeedsUAVSupport)
 {
 	// create output render target if necessary
 	ETextureCreateFlags OutputCreateFlags = TexCreate_Dynamic | (bInSRGB ? TexCreate_SRGB : TexCreate_None);
@@ -1498,7 +1555,7 @@ void FMediaTextureResource::CreateOutputRenderTarget(const FIntPoint & InDim, EP
 		OutputCreateFlags |= ETextureCreateFlags::RenderTargetable;
 	}
 
-	if ((InClearColor != CurrentClearColor) || !OutputTarget.IsValid() || (OutputTarget->GetSizeXY() != InDim) || (OutputTarget->GetFormat() != InPixelFormat) || ((OutputTarget->GetFlags() & OutputCreateFlags) != OutputCreateFlags) || CurrentNumMips != InNumMips)
+	if ((InClearColor != CurrentClearColor) || !IntermediateTarget.IsValid() || (IntermediateTarget->GetSizeXY() != InDim) || (IntermediateTarget->GetFormat() != InPixelFormat) || ((IntermediateTarget->GetFlags() & OutputCreateFlags) != OutputCreateFlags) || CurrentNumMips != InNumMips)
 	{
 		MipGenerationCache.SafeRelease();
 
@@ -1514,25 +1571,25 @@ void FMediaTextureResource::CreateOutputRenderTarget(const FIntPoint & InDim, EP
 			.SetClassName(ClassName)
 			.SetOwnerName(GetOwnerName());
 
-		OutputTarget = RHICreateTexture(Desc);
+		IntermediateTarget = RHICreateTexture(Desc);
 
-		OutputTarget->SetName(TEXT("MediaTextureResourceOutput"));
-		OutputTarget->SetOwnerName(GetOwnerName());
+		IntermediateTarget->SetName(TEXT("MediaTextureResourceOutput"));
+		IntermediateTarget->SetOwnerName(GetOwnerName());
 
 		CurrentClearColor = InClearColor;
 		CurrentNumMips = InNumMips;
 		UpdateResourceSize();
 
 		Cleared = false;
+		bRecreateOutputTarget = true;
 	}
-
-	if (RenderTargetTextureRHI != OutputTarget)
+	
+	if (RenderTargetTextureRHI != IntermediateTarget)
 	{
-		UpdateTextureReference(OutputTarget);
+		UpdateTextureReference(IntermediateTarget);
 	}
 }
 
-
 void FMediaTextureResource::CacheNextAvailableSampleTime(const TSharedPtr<FMediaTextureSampleSource, ESPMode::ThreadSafe>& InSampleQueue) const
 {
 	FTimespan SampleTime(FTimespan::MinValue());
diff --git a/Engine/Source/Runtime/MediaAssets/Private/Misc/MediaTextureResource.h b/Engine/Source/Runtime/MediaAssets/Private/Misc/MediaTextureResource.h
--- a/Engine/Source/Runtime/MediaAssets/Private/Misc/MediaTextureResource.h
+++ b/Engine/Source/Runtime/MediaAssets/Private/Misc/MediaTextureResource.h
@@ -174,10 +174,10 @@ protected:
 	void UpdateTextureReference(FRHITexture2D* NewTexture);
 
 	/**
-	 * Create/update output render target as needed
+	 * Create/update intermediate render target as needed. If no color conversion is needed, the RT will be used as the output.
 	 */
-	void CreateOutputRenderTarget(const FIntPoint & InDim, EPixelFormat InPixelFormat, bool bInSRGB, const FLinearColor & InClearColor, uint8 InNumMips, bool bNeedsUAVSupport);
-
+	void CreateIntermediateRenderTarget(const FIntPoint & InDim, EPixelFormat InPixelFormat, bool bInSRGB, const FLinearColor & InClearColor, uint8 InNumMips, bool bNeedsUAVSupport);
+	
 	/**
 	 * Caches next available sample from queue in MediaTexture owner to keep single consumer access
 	 *
@@ -214,7 +214,10 @@ private:
 	/** Input render target if the texture samples don't provide one (for conversions). */
 	TRefCountPtr<FRHITexture> InputTarget;
 
-	/** Output render target if the texture samples don't provide one. */
+	/** Holds the intermediate render target if the texture samples don't provide one, or final render target when not using a texture sample color converter. */
+	TRefCountPtr<FRHITexture> IntermediateTarget;
+
+	/** Output render target where the texture sample color converter writes. */
 	TRefCountPtr<FRHITexture> OutputTarget;
 
 	/** The media texture that owns this resource. */
@@ -258,4 +261,6 @@ private:
 	/** Colorspace to override standard proejct "working color space' */
 	TUniquePtr<UE::Color::FColorSpace> OverrideColorSpace;
 
+	/** Used to keep track of whether we should re-create the output target because the intermediate target has changed. */
+	bool bRecreateOutputTarget = false;
 };
