[Engine]

diff --git a/Engine/Source/Runtime/MediaAssets/Public/MediaPlayer.h b/Engine/Source/Runtime/MediaAssets/Public/MediaPlayer.h
--- a/Engine/Source/Runtime/MediaAssets/Public/MediaPlayer.h
+++ b/Engine/Source/Runtime/MediaAssets/Public/MediaPlayer.h
@@ -74,8 +74,8 @@ enum class EMediaPlayerTrack : uint8
 };
 
 
-UCLASS(BlueprintType, hidecategories = (Object))
-class MEDIAASSETS_API UMediaTimeStampInfo
+UCLASS(BlueprintType, hidecategories = (Object), MinimalAPI)
+class UMediaTimeStampInfo
 	: public UObject
 {
 	GENERATED_BODY()
@@ -91,7 +91,7 @@ public:
 };
 
 USTRUCT(BlueprintType)
-struct MEDIAASSETS_API FMediaMetadataItemBPT
+struct FMediaMetadataItemBPT
 {
 	GENERATED_BODY()
 
@@ -110,7 +110,7 @@ public:
 };
 
 USTRUCT(BlueprintType)
-struct MEDIAASSETS_API FMediaMetadataItemsBPT
+struct FMediaMetadataItemsBPT
 {
 	GENERATED_BODY()
 public:
@@ -121,16 +121,16 @@ public:
 /**
  * Proxy implementation for Metasound integration
  */
-class MEDIAASSETS_API FMediaPlayerProxy : public Audio::TProxyData<FMediaPlayerProxy>
+class FMediaPlayerProxy : public Audio::TProxyData<FMediaPlayerProxy>
 {
 public:
 	IMPL_AUDIOPROXY_CLASS(FMediaPlayerProxy);
 
-	explicit FMediaPlayerProxy(UMediaPlayer* Player);
+	MEDIAASSETS_API explicit FMediaPlayerProxy(UMediaPlayer* Player);
 
 	FMediaPlayerProxy(const FMediaPlayerProxy& Other) = default;
 
-	~FMediaPlayerProxy();
+	MEDIAASSETS_API ~FMediaPlayerProxy();
 
 	TSharedPtr<FMediaPlayerFacade, ESPMode::ThreadSafe> GetPlayerFacade() const
 	{
@@ -146,8 +146,8 @@ using FMediaPlayerProxyPtr = TSharedPtr<FMediaPlayerProxy, ESPMode::ThreadSafe>;
 /**
  * Implements a media player asset that can play movies and other media sources.
  */
-UCLASS(BlueprintType, hidecategories=(Object))
-class MEDIAASSETS_API UMediaPlayer
+UCLASS(BlueprintType, hidecategories=(Object), MinimalAPI)
+class UMediaPlayer
 	: public UObject
 	, public IAudioProxyDataFactory
 {
@@ -164,7 +164,7 @@ public:
 	 * @see CanPlay, Pause
 	 */
 	UFUNCTION(BlueprintCallable, Category="Media|MediaPlayer")
-	bool CanPause() const;
+	MEDIAASSETS_API bool CanPause() const;
 
 	/**
 	 * Check whether the specified media source can be played by this player.
@@ -177,7 +177,7 @@ public:
 	 * @see CanPlayUrl, SetDesiredPlayerName
 	 */
 	UFUNCTION(BlueprintCallable, Category = "Media|MediaPlayer")
-	bool CanPlaySource(UMediaSource* MediaSource);
+	MEDIAASSETS_API bool CanPlaySource(UMediaSource* MediaSource);
 
 	/**
 	 * Check whether the specified URL can be played by this player.
@@ -189,7 +189,7 @@ public:
 	 * @see CanPlaySource, SetDesiredPlayerName
 	 */
 	UFUNCTION(BlueprintCallable, Category="Media|MediaPlayer")
-	bool CanPlayUrl(const FString& Url);
+	MEDIAASSETS_API bool CanPlayUrl(const FString& Url);
 
 	/**
 	 * Close the currently open media, if any.
@@ -197,7 +197,7 @@ public:
 	 * @see OnMediaClosed, OpenPlaylist, OpenPlaylistIndex, OpenSource, OpenUrl, Pause, Play
 	 */
 	UFUNCTION(BlueprintCallable, Category="Media|MediaPlayer")
-	void Close();
+	MEDIAASSETS_API void Close();
 
 	/**
 	 * Get the number of channels in the specified audio track.
@@ -208,7 +208,7 @@ public:
 	 * @see GetAudioTrackSampleRate, GetAudioTrackType
 	 */
 	UFUNCTION(BlueprintCallable, Category = "Media|MediaPlayer")
-	int32 GetAudioTrackChannels(int32 TrackIndex, int32 FormatIndex) const;
+	MEDIAASSETS_API int32 GetAudioTrackChannels(int32 TrackIndex, int32 FormatIndex) const;
 
 	/**
 	 * Get the sample rate of the specified audio track.
@@ -219,7 +219,7 @@ public:
 	 * @see GetAudioTrackChannels, GetAudioTrackType
 	 */
 	UFUNCTION(BlueprintCallable, Category = "Media|MediaPlayer")
-	int32 GetAudioTrackSampleRate(int32 TrackIndex, int32 FormatIndex) const;
+	MEDIAASSETS_API int32 GetAudioTrackSampleRate(int32 TrackIndex, int32 FormatIndex) const;
 
 	/**
 	 * Get the type of the specified audio track format.
@@ -230,7 +230,7 @@ public:
 	 * @see GetAudioTrackSampleRate, GetAudioTrackSampleRate
 	 */
 	UFUNCTION(BlueprintCallable, Category="Media|MediaPlayer")
-	FString GetAudioTrackType(int32 TrackIndex, int32 FormatIndex) const;
+	MEDIAASSETS_API FString GetAudioTrackType(int32 TrackIndex, int32 FormatIndex) const;
 
 	/**
 	 * Get the name of the current desired native player.
@@ -239,7 +239,7 @@ public:
 	 * @see SetDesiredPlayerName
 	 */
 	UFUNCTION(BlueprintCallable, Category = "Media|MediaPlayer")
-	FName GetDesiredPlayerName() const;
+	MEDIAASSETS_API FName GetDesiredPlayerName() const;
 
 	/**
 	 * Get the media's duration.
@@ -248,7 +248,7 @@ public:
 	 * @see GetTime, Seek
 	 */
 	UFUNCTION(BlueprintCallable, Category="Media|MediaPlayer")
-	FTimespan GetDuration() const;
+	MEDIAASSETS_API FTimespan GetDuration() const;
 
 	/**
 	 * Get the current horizontal field of view (only for 360 videos).
@@ -257,7 +257,7 @@ public:
 	 * @see GetVerticalFieldOfView, GetViewRotation, SetHorizontalFieldOfView
 	 */
 	UFUNCTION(BlueprintCallable, Category="Media|MediaPlayer")
-	float GetHorizontalFieldOfView() const;
+	MEDIAASSETS_API float GetHorizontalFieldOfView() const;
 
 	/**
 	 * Get the human readable name of the currently loaded media source.
@@ -266,7 +266,7 @@ public:
 	 * @see GetPlayerName, GetUrl
 	 */
 	UFUNCTION(BlueprintCallable, Category="Media|MediaPlayer")
-	virtual FText GetMediaName() const;
+	MEDIAASSETS_API virtual FText GetMediaName() const;
 
 	/**
 	 * Get the number of tracks of the given type.
@@ -276,7 +276,7 @@ public:
 	 * @see GetNumTrackFormats, GetSelectedTrack, SelectTrack
 	 */
 	UFUNCTION(BlueprintCallable, Category="Media|MediaPlayer")
-	int32 GetNumTracks(EMediaPlayerTrack TrackType) const;
+	MEDIAASSETS_API int32 GetNumTracks(EMediaPlayerTrack TrackType) const;
 
 	/**
 	 * Get the number of formats of the specified track.
@@ -287,7 +287,7 @@ public:
 	 * @see GetNumTracks, GetSelectedTrack, SelectTrack
 	 */
 	UFUNCTION(BlueprintCallable, Category="Media|MediaPlayer")
-	int32 GetNumTrackFormats(EMediaPlayerTrack TrackType, int32 TrackIndex) const;
+	MEDIAASSETS_API int32 GetNumTrackFormats(EMediaPlayerTrack TrackType, int32 TrackIndex) const;
 
 	/**
 	 * Get the name of the current native media player.
@@ -296,7 +296,7 @@ public:
 	 * @see GetMediaName
 	 */
 	UFUNCTION(BlueprintCallable, Category="Media|MediaPlayer")
-	FName GetPlayerName() const;
+	MEDIAASSETS_API FName GetPlayerName() const;
 
 	/**
 	 * Get the current play list.
@@ -333,7 +333,7 @@ public:
 	 * @see SetRate, SupportsRate
 	 */
 	UFUNCTION(BlueprintCallable, Category="Media|MediaPlayer")
-	float GetRate() const;
+	MEDIAASSETS_API float GetRate() const;
 
 	/**
 	 * Get the index of the currently selected track of the given type.
@@ -343,7 +343,7 @@ public:
 	 * @see GetNumTracks, GetTrackFormat, SelectTrack
 	 */
 	UFUNCTION(BlueprintCallable, Category="Media|MediaPlayer")
-	int32 GetSelectedTrack(EMediaPlayerTrack TrackType) const;
+	MEDIAASSETS_API int32 GetSelectedTrack(EMediaPlayerTrack TrackType) const;
 
 	/**
 	 * Get the supported playback rates.
@@ -353,7 +353,7 @@ public:
 	 * @see SetRate, SupportsRate
 	 */
 	UFUNCTION(BlueprintCallable, Category="Media|MediaPlayer")
-	void GetSupportedRates(TArray<FFloatRange>& OutRates, bool Unthinned) const;
+	MEDIAASSETS_API void GetSupportedRates(TArray<FFloatRange>& OutRates, bool Unthinned) const;
 
 	/**
 	 * Get the media's current playback time.
@@ -362,7 +362,7 @@ public:
 	 * @see GetDuration, Seek
 	 */
 	UFUNCTION(BlueprintCallable, Category="Media|MediaPlayer")
-	FTimespan GetTime() const;
+	MEDIAASSETS_API FTimespan GetTime() const;
 
 	/**
 	 * Get the media's current playback time as appropriate for display.
@@ -371,7 +371,7 @@ public:
 	 * @see GetDuration, Seek
 	 */
 	UFUNCTION(BlueprintCallable, Category = "Media|MediaPlayer")
-	FTimespan GetDisplayTime() const;
+	MEDIAASSETS_API FTimespan GetDisplayTime() const;
 
 	/**
 	 * Get the media's current playback timestamp.
@@ -380,7 +380,7 @@ public:
 	 * @see GetDuration, Seek
 	 */
 	UFUNCTION(BlueprintCallable, Category = "Media|MediaPlayer")
-	UMediaTimeStampInfo* GetTimeStamp() const;
+	MEDIAASSETS_API UMediaTimeStampInfo* GetTimeStamp() const;
 
 	/**
 	 * Get the media's current playback timestamp as appropriate for display.
@@ -389,7 +389,7 @@ public:
 	 * @see GetDuration, Seek
 	 */
 	UFUNCTION(BlueprintCallable, Category = "Media|MediaPlayer")
-	UMediaTimeStampInfo* GetDisplayTimeStamp() const;
+	MEDIAASSETS_API UMediaTimeStampInfo* GetDisplayTimeStamp() const;
 
 	/**
 	 * Get the human readable name of the specified track.
@@ -400,7 +400,7 @@ public:
 	 * @see GetNumTracks, GetTrackLanguage
 	 */
 	UFUNCTION(BlueprintCallable, Category = "Media|MediaPlayer")
-	FText GetTrackDisplayName(EMediaPlayerTrack TrackType, int32 TrackIndex) const;
+	MEDIAASSETS_API FText GetTrackDisplayName(EMediaPlayerTrack TrackType, int32 TrackIndex) const;
 
 	/**
 	 * Get the index of the active format of the specified track type.
@@ -411,7 +411,7 @@ public:
 	 * @see GetNumTrackFormats, GetSelectedTrack, SetTrackFormat
 	 */
 	UFUNCTION(BlueprintCallable, Category="Media|MediaPlayer")
-	int32 GetTrackFormat(EMediaPlayerTrack TrackType, int32 TrackIndex) const;
+	MEDIAASSETS_API int32 GetTrackFormat(EMediaPlayerTrack TrackType, int32 TrackIndex) const;
 
 	/**
 	 * Get the language tag of the specified track.
@@ -422,7 +422,7 @@ public:
 	 * @see GetNumTracks, GetTrackDisplayName
 	 */
 	UFUNCTION(BlueprintCallable, Category="Media|MediaPlayer")
-	FString GetTrackLanguage(EMediaPlayerTrack TrackType, int32 TrackIndex) const;
+	MEDIAASSETS_API FString GetTrackLanguage(EMediaPlayerTrack TrackType, int32 TrackIndex) const;
 
 	/**
 	 * Get the URL of the currently loaded media, if any.
@@ -431,7 +431,7 @@ public:
 	 * @see OpenUrl
 	 */
 	UFUNCTION(BlueprintCallable, Category="Media|MediaPlayer")
-	const FString& GetUrl() const;
+	MEDIAASSETS_API const FString& GetUrl() const;
 
 	/**
 	 * Get the current vertical field of view (only for 360 videos).
@@ -440,7 +440,7 @@ public:
 	 * @see GetHorizontalFieldOfView, GetViewRotation, SetVerticalFieldOfView
 	 */
 	UFUNCTION(BlueprintCallable, Category="Media|MediaPlayer")
-	float GetVerticalFieldOfView() const;
+	MEDIAASSETS_API float GetVerticalFieldOfView() const;
 
 	/**
 	 * Get the aspect ratio of the specified video track.
@@ -451,7 +451,7 @@ public:
 	 * @see GetVideoTrackDimensions, GetVideoTrackFrameRate, GetVideoTrackFrameRates, GetVideoTrackType
 	 */
 	UFUNCTION(BlueprintCallable, Category = "Media|MediaPlayer")
-	float GetVideoTrackAspectRatio(int32 TrackIndex, int32 FormatIndex) const;
+	MEDIAASSETS_API float GetVideoTrackAspectRatio(int32 TrackIndex, int32 FormatIndex) const;
 
 	/**
 	 * Get the current dimensions of the specified video track.
@@ -462,7 +462,7 @@ public:
 	 * @see GetVideoTrackAspectRatio, GetVideoTrackFrameRate, GetVideoTrackFrameRates, GetVideoTrackType
 	 */
 	UFUNCTION(BlueprintCallable, Category="Media|MediaPlayer")
-	FIntPoint GetVideoTrackDimensions(int32 TrackIndex, int32 FormatIndex) const;
+	MEDIAASSETS_API FIntPoint GetVideoTrackDimensions(int32 TrackIndex, int32 FormatIndex) const;
 
 	/**
 	 * Get the frame rate of the specified video track.
@@ -473,7 +473,7 @@ public:
 	 * @see GetVideoTrackAspectRatio, GetVideoTrackDimensions, GetVideoTrackFrameRates, GetVideoTrackType, SetVideoTrackFrameRate
 	 */
 	UFUNCTION(BlueprintCallable, Category="Media|MediaPlayer")
-	float GetVideoTrackFrameRate(int32 TrackIndex, int32 FormatIndex) const;
+	MEDIAASSETS_API float GetVideoTrackFrameRate(int32 TrackIndex, int32 FormatIndex) const;
 
 	/**
 	 * Get the supported range of frame rates of the specified video track.
@@ -484,7 +484,7 @@ public:
 	 * @see GetVideoTrackAspectRatio, GetVideoTrackDimensions, GetVideoTrackFrameRate, GetVideoTrackType
 	 */
 	UFUNCTION(BlueprintCallable, Category="Media|MediaPlayer")
-	FFloatRange GetVideoTrackFrameRates(int32 TrackIndex, int32 FormatIndex) const;
+	MEDIAASSETS_API FFloatRange GetVideoTrackFrameRates(int32 TrackIndex, int32 FormatIndex) const;
 
 	/**
 	 * Get the type of the specified video track format.
@@ -495,7 +495,7 @@ public:
 	 * @see GetVideoTrackAspectRatio, GetVideoTrackDimensions, GetVideoTrackFrameRate, GetVideoTrackFrameRates
 	 */
 	UFUNCTION(BlueprintCallable, Category="Media|MediaPlayer")
-	FString GetVideoTrackType(int32 TrackIndex, int32 FormatIndex) const;
+	MEDIAASSETS_API FString GetVideoTrackType(int32 TrackIndex, int32 FormatIndex) const;
 
 	/**
 	 * Get the current view rotation (only for 360 videos).
@@ -504,7 +504,7 @@ public:
 	 * @see GetHorizontalFieldOfView, GetVerticalFieldOfView, SetViewRotation
 	 */
 	UFUNCTION(BlueprintCallable, Category="Media|MediaPlayer")
-	FRotator GetViewRotation() const;
+	MEDIAASSETS_API FRotator GetViewRotation() const;
 
 	/**
 	 * Delay of the player's time.
@@ -513,7 +513,7 @@ public:
 	 * @see SetTimeDelay
 	 */
 	UFUNCTION(BlueprintCallable, Category="Media|MediaPlayer")
-	FTimespan GetTimeDelay() const;
+	MEDIAASSETS_API FTimespan GetTimeDelay() const;
 
 	/**
 	 * Check whether the player is in an error state.
@@ -526,7 +526,7 @@ public:
 	 * @see IsReady
 	 */
 	UFUNCTION(BlueprintCallable, Category="Media|MediaPlayer")
-	bool HasError() const;
+	MEDIAASSETS_API bool HasError() const;
 
 	/**
 	 * Check whether playback is buffering data.
@@ -535,7 +535,7 @@ public:
 	 * @see IsConnecting, IsLooping, IsPaused, IsPlaying, IsPreparing, IsReady
 	 */
 	UFUNCTION(BlueprintCallable, Category="Media|MediaPlayer")
-	bool IsBuffering() const;
+	MEDIAASSETS_API bool IsBuffering() const;
 
 	/**
 	 * Check whether the player is currently connecting to a media source.
@@ -544,7 +544,7 @@ public:
 	 * @see IsBuffering, IsLooping, IsPaused, IsPlaying, IsPreparing, IsReady
 	 */
 	UFUNCTION(BlueprintCallable, Category="Media|MediaPlayer")
-	bool IsConnecting() const;
+	MEDIAASSETS_API bool IsConnecting() const;
 
 	/**
 	 * Check whether playback is looping.
@@ -553,7 +553,7 @@ public:
 	 * @see IsBuffering, IsConnecting, IsPaused, IsPlaying, IsPreparing, IsReady, SetLooping
 	 */
 	UFUNCTION(BlueprintCallable, Category="Media|MediaPlayer")
-	bool IsLooping() const;
+	MEDIAASSETS_API bool IsLooping() const;
 
 	/**
 	 * Check whether playback is currently paused.
@@ -562,7 +562,7 @@ public:
 	 * @see CanPause, IsBuffering, IsConnecting, IsLooping, IsPaused, IsPlaying, IsPreparing, IsReady, Pause
 	 */
 	UFUNCTION(BlueprintCallable, Category="Media|MediaPlayer")
-	bool IsPaused() const;
+	MEDIAASSETS_API bool IsPaused() const;
 
 	/**
 	 * Check whether playback has started.
@@ -571,7 +571,7 @@ public:
 	 * @see CanPlay, IsBuffering, IsConnecting, IsLooping, IsPaused, IsPlaying, IsPreparing, IsReady, Play
 	 */
 	UFUNCTION(BlueprintCallable, Category="Media|MediaPlayer")
-	bool IsPlaying() const;
+	MEDIAASSETS_API bool IsPlaying() const;
 
 	/**
 	 * Check whether the media is currently opening or buffering.
@@ -580,7 +580,7 @@ public:
 	 * @see CanPlay, IsBuffering, IsConnecting, IsLooping, IsPaused, IsPlaying, IsReady, Play
 	 */
 	UFUNCTION(BlueprintCallable, Category="Media|MediaPlayer")
-	bool IsPreparing() const;
+	MEDIAASSETS_API bool IsPreparing() const;
 
 	/**
 	 * Whether media is currently closed.
@@ -588,7 +588,7 @@ public:
 	 * @return true if media is closed, false otherwise.
 	 */
 	UFUNCTION(BlueprintCallable, Category="Media|MediaPlayer")
-	bool IsClosed() const;
+	MEDIAASSETS_API bool IsClosed() const;
 
 	/**
 	 * Check whether media is ready for playback.
@@ -600,7 +600,7 @@ public:
 	 * @see HasError, IsBuffering, IsConnecting, IsLooping, IsPaused, IsPlaying, IsPreparing
 	 */
 	UFUNCTION(BlueprintCallable, Category="Media|MediaPlayer")
-	bool IsReady() const;
+	MEDIAASSETS_API bool IsReady() const;
 
 	/**
 	 * Open the next item in the current play list.
@@ -612,7 +612,7 @@ public:
 	 * @see Close, OpenUrl, OpenSource, Play, Previous, SetPlaylist
 	 */
 	UFUNCTION(BlueprintCallable, Category="Media|MediaPlayer")
-	bool Next();
+	MEDIAASSETS_API bool Next();
 
 	/**
 	 * Opens the specified media file path.
@@ -627,7 +627,7 @@ public:
 	 * @see GetUrl, Close, OpenPlaylist, OpenPlaylistIndex, OpenSource, OpenUrl, Reopen
 	 */
 	UFUNCTION(BlueprintCallable, Category="Media|MediaPlayer")
-	bool OpenFile(const FString& FilePath);
+	MEDIAASSETS_API bool OpenFile(const FString& FilePath);
 
 	/**
 	 * Open the first media source in the specified play list.
@@ -651,7 +651,7 @@ public:
 	 * @see Close, OpenFile, OpenPlaylist, OpenSource, OpenUrl, Reopen
 	 */
 	UFUNCTION(BlueprintCallable, Category="Media|MediaPlayer")
-	bool OpenPlaylistIndex(UMediaPlaylist* InPlaylist, int32 Index);
+	MEDIAASSETS_API bool OpenPlaylistIndex(UMediaPlaylist* InPlaylist, int32 Index);
 
 	/**
 	 * Open the specified media source.
@@ -666,7 +666,7 @@ public:
 	 * @see Close, OpenFile, OpenPlaylist, OpenPlaylistIndex, OpenUrl, Reopen
 	 */
 	UFUNCTION(BlueprintCallable, Category="Media|MediaPlayer")
-	bool OpenSource(UMediaSource* MediaSource);
+	MEDIAASSETS_API bool OpenSource(UMediaSource* MediaSource);
 
 	/**
 	 * Open the specified media source with supplied options applied.
@@ -682,7 +682,7 @@ public:
 	 * @see Close, OpenFile, OpenPlaylist, OpenPlaylistIndex, OpenUrl, Reopen
 	 */
 	UFUNCTION(BlueprintCallable, Category="Media|MediaPlayer")
-	bool OpenSourceWithOptions(UMediaSource* MediaSource, const FMediaPlayerOptions& Options);
+	MEDIAASSETS_API bool OpenSourceWithOptions(UMediaSource* MediaSource, const FMediaPlayerOptions& Options);
 
 	/**
 	 * Open the specified media source with options using a latent action.
@@ -695,7 +695,7 @@ public:
 	 * @see Close, OpenFile, OpenPlaylist, OpenPlaylistIndex, OpenUrl, Reopen
 	 */
 	UFUNCTION(BlueprintCallable, Category = "Media|MediaPlayer", meta = (Latent, LatentInfo = "LatentInfo", WorldContext = "WorldContextObject"))
-	void OpenSourceLatent(const UObject* WorldContextObject, struct FLatentActionInfo LatentInfo, UMediaSource* MediaSource, const FMediaPlayerOptions& Options, bool& bSuccess);
+	MEDIAASSETS_API void OpenSourceLatent(const UObject* WorldContextObject, struct FLatentActionInfo LatentInfo, UMediaSource* MediaSource, const FMediaPlayerOptions& Options, bool& bSuccess);
 
 	/**
 	 * Opens the specified media URL.
@@ -710,7 +710,7 @@ public:
 	 * @see GetUrl, Close, OpenFile, OpenPlaylist, OpenPlaylistIndex, OpenSource, Reopen
 	 */
 	UFUNCTION(BlueprintCallable, Category="Media|MediaPlayer")
-	bool OpenUrl(const FString& Url);
+	MEDIAASSETS_API bool OpenUrl(const FString& Url);
 
 	/**
 	 * Pauses media playback.
@@ -721,7 +721,7 @@ public:
 	 * @see CanPause, Close, Next, Play, Previous, Rewind, Seek
 	 */
 	UFUNCTION(BlueprintCallable, Category="Media|MediaPlayer")
-	bool Pause();
+	MEDIAASSETS_API bool Pause();
 
 	/**
 	 * Starts media playback.
@@ -732,13 +732,13 @@ public:
 	 * @see CanPlay, GetRate, Next, Pause, Previous, SetRate
 	 */
 	UFUNCTION(BlueprintCallable, Category="Media|MediaPlayer")
-	bool Play();
+	MEDIAASSETS_API bool Play();
 
 	/**
 	 * Starts playback from the media opened event, but can be used elsewhere.
 	 */
 	UFUNCTION(BlueprintCallable, Category="Media|MediaPlayer")
-	void PlayAndSeek();
+	MEDIAASSETS_API void PlayAndSeek();
 
 	/**
 	 * Open the previous item in the current play list.
@@ -750,7 +750,7 @@ public:
 	 * @see Close, Next, OpenUrl, OpenSource, Play, SetPlaylist
 	 */
 	UFUNCTION(BlueprintCallable, Category="Media|MediaPlayer")
-	bool Previous();
+	MEDIAASSETS_API bool Previous();
 
 	/**
 	 * Reopens the currently opened media or play list.
@@ -759,7 +759,7 @@ public:
 	 * @see Close, Open, OpenFile, OpenPlaylist, OpenPlaylistIndex, OpenSource, OpenUrl
 	 */
 	UFUNCTION(BlueprintCallable, Category="Media|MediaPlayer")
-	bool Reopen();
+	MEDIAASSETS_API bool Reopen();
 
 	/**
 	 * Rewinds the media to the beginning.
@@ -770,7 +770,7 @@ public:
 	 * @see GetTime, Seek
 	 */
 	UFUNCTION(BlueprintCallable, Category="Media|MediaPlayer")
-	bool Rewind();
+	MEDIAASSETS_API bool Rewind();
 
 	/**
 	 * Seeks to the specified playback time.
@@ -780,7 +780,7 @@ public:
 	 * @see GetTime, Rewind
 	 */
 	UFUNCTION(BlueprintCallable, Category="Media|MediaPlayer")
-	bool Seek(const FTimespan& Time);
+	MEDIAASSETS_API bool Seek(const FTimespan& Time);
 
 	/**
 	 * Select the active track of the given type.
@@ -795,7 +795,7 @@ public:
 	 * @see GetNumTracks, GetSelectedTrack, SetTrackFormat
 	 */
 	UFUNCTION(BlueprintCallable, Category="Media|MediaPlayer")
-	bool SelectTrack(EMediaPlayerTrack TrackType, int32 TrackIndex);
+	MEDIAASSETS_API bool SelectTrack(EMediaPlayerTrack TrackType, int32 TrackIndex);
 
 	/**
 	 * Set the time on which to block.
@@ -806,7 +806,7 @@ public:
 	 * @param Time The time to block on, or FTimespan::MinValue to disable.
 	 */
 	UFUNCTION(BlueprintCallable, Category = "Media|MediaPlayer")
-	void SetBlockOnTime(const FTimespan& Time);
+	MEDIAASSETS_API void SetBlockOnTime(const FTimespan& Time);
 
 	/**
 	 * Set the time range on which to block.
@@ -816,7 +816,7 @@ public:
 	 *
 	 * @param TimeRange The time range to block on, or empty range to disable.
 	 */
-	void SetBlockOnTimeRange(const TRange<FTimespan>& TimeRange);
+	MEDIAASSETS_API void SetBlockOnTimeRange(const TRange<FTimespan>& TimeRange);
 
 	/**
 	 * Set the name of the desired native player.
@@ -825,7 +825,7 @@ public:
 	 * @see GetDesiredPlayerName
 	 */
 	UFUNCTION(BlueprintCallable, Category = "Media|MediaPlayer")
-	void SetDesiredPlayerName(FName PlayerName);
+	MEDIAASSETS_API void SetDesiredPlayerName(FName PlayerName);
 
 	/**
 	 * Enables or disables playback looping.
@@ -835,7 +835,7 @@ public:
 	 * @see IsLooping
 	 */
 	UFUNCTION(BlueprintCallable, Category="Media|MediaPlayer")
-	bool SetLooping(bool Looping);
+	MEDIAASSETS_API bool SetLooping(bool Looping);
 
 	/**
 	 * Sets the media options used by the player.
@@ -843,7 +843,7 @@ public:
 	 * @param Options Options to pass to the player.
 	 */
 	UFUNCTION(BlueprintCallable, Category = "Media|MediaPlayer")
-	void SetMediaOptions(const UMediaSource* Options);
+	MEDIAASSETS_API void SetMediaOptions(const UMediaSource* Options);
 
 	/**
 	 * Changes the media's playback rate.
@@ -853,7 +853,7 @@ public:
 	 * @see GetRate, SupportsRate
 	 */
 	UFUNCTION(BlueprintCallable, Category="Media|MediaPlayer")
-	bool SetRate(float Rate);
+	MEDIAASSETS_API bool SetRate(float Rate);
 
 
 	/**
@@ -867,7 +867,7 @@ public:
 	 * @see NativeAudioOut
 	 */
 	UFUNCTION(BlueprintCallable, Category = "Media|MediaPlayer")
-	bool SetNativeVolume(float Volume);
+	MEDIAASSETS_API bool SetNativeVolume(float Volume);
 
 	/**
 	 * Set the format on the specified track.
@@ -883,7 +883,7 @@ public:
 	 * @see GetNumTrackFormats, GetNumTracks, GetTrackFormat, SelectTrack
 	 */
 	UFUNCTION(BlueprintCallable, Category="Media|MediaPlayer")
-	bool SetTrackFormat(EMediaPlayerTrack TrackType, int32 TrackIndex, int32 FormatIndex);
+	MEDIAASSETS_API bool SetTrackFormat(EMediaPlayerTrack TrackType, int32 TrackIndex, int32 FormatIndex);
 
 	/**
 	 * Set the frame rate of the specified video track.
@@ -895,7 +895,7 @@ public:
 	 * @see GetVideoTrackAspectRatio, GetVideoTrackDimensions, GetVideoTrackFrameRate, GetVideoTrackFrameRates, GetVideoTrackType
 	 */
 	UFUNCTION(BlueprintCallable, Category="Media|MediaPlayer")
-	bool SetVideoTrackFrameRate(int32 TrackIndex, int32 FormatIndex, float FrameRate);
+	MEDIAASSETS_API bool SetVideoTrackFrameRate(int32 TrackIndex, int32 FormatIndex, float FrameRate);
 
 	/**
 	 * Set the field of view (only for 360 videos).
@@ -907,7 +907,7 @@ public:
 	 * @see GetHorizontalFieldOfView, GetVerticalFieldOfView, SetViewRotation
 	 */
 	UFUNCTION(BlueprintCallable, Category = "Media|MediaPlayer")
-	bool SetViewField(float Horizontal, float Vertical, bool Absolute);
+	MEDIAASSETS_API bool SetViewField(float Horizontal, float Vertical, bool Absolute);
 
 	/**
 	 * Set the view's rotation (only for 360 videos).
@@ -918,7 +918,7 @@ public:
 	 * @see GetViewRotation, SetViewField
 	 */
 	UFUNCTION(BlueprintCallable, Category = "Media|MediaPlayer")
-	bool SetViewRotation(const FRotator& Rotation, bool Absolute);
+	MEDIAASSETS_API bool SetViewRotation(const FRotator& Rotation, bool Absolute);
 
 	/**
 	 * Delay of the player's time.
@@ -930,7 +930,7 @@ public:
 	 * @see GetTimeDelay
 	 */
 	UFUNCTION(BlueprintCallable, Category = "Media|MediaPlayer")
-	void SetTimeDelay(FTimespan TimeDelay);
+	MEDIAASSETS_API void SetTimeDelay(FTimespan TimeDelay);
 
 	/**
 	 * Check whether the specified playback rate is supported.
@@ -940,7 +940,7 @@ public:
 	 * @see SupportsScrubbing, SupportsSeeking
 	 */
 	UFUNCTION(BlueprintCallable, Category="Media|MediaPlayer")
-	bool SupportsRate(float Rate, bool Unthinned) const;
+	MEDIAASSETS_API bool SupportsRate(float Rate, bool Unthinned) const;
 
 	/**
 	 * Check whether the currently loaded media supports scrubbing.
@@ -949,7 +949,7 @@ public:
 	 * @see SupportsRate, SupportsSeeking
 	 */
 	UFUNCTION(BlueprintCallable, Category="Media|MediaPlayer")
-	bool SupportsScrubbing() const;
+	MEDIAASSETS_API bool SupportsScrubbing() const;
 
 	/**
 	 * Check whether the currently loaded media can jump to a certain position.
@@ -958,7 +958,7 @@ public:
 	 * @see SupportsRate, SupportsScrubbing
 	 */
 	UFUNCTION(BlueprintCallable, Category="Media|MediaPlayer")
-	bool SupportsSeeking() const;
+	MEDIAASSETS_API bool SupportsSeeking() const;
 
 public:
 
@@ -1042,7 +1042,7 @@ public:
 	 * but if you want to clean up resources before garbage collection (e.g. you are in the editor)
 	 * then this could be helpful.
 	 */
-	void CleanUpBeforeDestroy();
+	MEDIAASSETS_API void CleanUpBeforeDestroy();
 
 	/**
 	 * Get the Guid associated with this media player
@@ -1063,7 +1063,7 @@ public:
 	 * @param	InfoName		Name of the information we want.
 	 * @returns					Requested information, or empty if not available.
 	 */
-	FVariant GetMediaInfo(FName InfoName) const;
+	MEDIAASSETS_API FVariant GetMediaInfo(FName InfoName) const;
 
 	/**
 	 * Templated version of GetMediaInfo.
@@ -1094,7 +1094,7 @@ public:
 	 * @return Map with arrays of IMediaMetaDataItem entries describing any metadata found in the current stream
 	 * @note Listen to EMediaEvent::MetadataChanged to catch updates to this data
 	 */
-	TSharedPtr<TMap<FString, TArray<TUniquePtr<IMediaMetadataItem>>>, ESPMode::ThreadSafe> GetMediaMetadata() const;
+	MEDIAASSETS_API TSharedPtr<TMap<FString, TArray<TUniquePtr<IMediaMetadataItem>>>, ESPMode::ThreadSafe> GetMediaMetadata() const;
 
 	/**
 	 * This is the blueprint accessible version of the GetMediaMetadata.
@@ -1102,25 +1102,25 @@ public:
 	 * @note Listen to EMediaEvent::MetadataChanged to catch updates to this data
 	 */
 	UFUNCTION(BlueprintCallable, Category = "Media|Metadata")
-	TMap<FString, FMediaMetadataItemsBPT> GetMediaMetadataItems() const;
+	MEDIAASSETS_API TMap<FString, FMediaMetadataItemsBPT> GetMediaMetadataItems() const;
 	
 	/**
 	 * Get the media player facade that manages low-level media players
 	 *
 	 * @return The media player facade.
 	 */
-	TSharedRef<FMediaPlayerFacade, ESPMode::ThreadSafe> GetPlayerFacade() const;
+	MEDIAASSETS_API TSharedRef<FMediaPlayerFacade, ESPMode::ThreadSafe> GetPlayerFacade() const;
 
 	/**
 	 * Register player with media module to be ticked
 	 *
 	 */
-	void RegisterWithMediaModule();
+	MEDIAASSETS_API void RegisterWithMediaModule();
 
 	/** 
 	 * When the player goes out of scope, make sure to clean up the clock sink
 	 */
-	void UnregisterWithMediaModule();
+	MEDIAASSETS_API void UnregisterWithMediaModule();
 
 	/**
 	 * Get the current play list.
@@ -1151,29 +1151,29 @@ public:
 	 *
 	 * @see ResumePIE
 	 */
-	void PausePIE();
+	MEDIAASSETS_API void PausePIE();
 
 	/**
 	 * Called when PIE has been resumed.
 	 *
 	 * @see PausePIE
 	 */
-	void ResumePIE();
+	MEDIAASSETS_API void ResumePIE();
 #endif
 
 public:
 
 	//~ UObject interface
 
-	virtual void BeginDestroy() override;
+	MEDIAASSETS_API virtual void BeginDestroy() override;
 
-	virtual FString GetDesc() override;
-	virtual void PostDuplicate(bool bDuplicateForPIE) override;
-	virtual void PostInitProperties() override;
-	virtual void PostLoad() override;
+	MEDIAASSETS_API virtual FString GetDesc() override;
+	MEDIAASSETS_API virtual void PostDuplicate(bool bDuplicateForPIE) override;
+	MEDIAASSETS_API virtual void PostInitProperties() override;
+	MEDIAASSETS_API virtual void PostLoad() override;
 
 #if WITH_EDITOR
-	virtual void PostEditChangeProperty(FPropertyChangedEvent& PropertyChangedEvent) override;
+	MEDIAASSETS_API virtual void PostEditChangeProperty(FPropertyChangedEvent& PropertyChangedEvent) override;
 #endif
 
 public:
@@ -1182,9 +1182,9 @@ public:
 	 * They also decsribe the type of data that will be needed to retrieve the information.
 	 */
 	 /** Number of mip map levels in the source. Int32.*/
-	static FLazyName MediaInfoNameSourceNumMips;
+	static MEDIAASSETS_API FLazyName MediaInfoNameSourceNumMips;
 	/** Number of tiles (X and Y) in the source. IntPoint.*/
-	static FLazyName MediaInfoNameSourceNumTiles;
+	static MEDIAASSETS_API FLazyName MediaInfoNameSourceNumTiles;
 
 	/**
 	 * Duration of samples to cache ahead of the play head.
@@ -1250,7 +1250,7 @@ public:
 	// IAudioProxyDataFactory interface (for Meta Sound support)
 	//
 
-	virtual TSharedPtr<Audio::IProxyData> CreateProxyData(const Audio::FProxyDataInitParams& InitParams) override;
+	MEDIAASSETS_API virtual TSharedPtr<Audio::IProxyData> CreateProxyData(const Audio::FProxyDataInitParams& InitParams) override;
 
 protected:
 
@@ -1339,16 +1339,16 @@ private:
 	virtual bool CanBeInCluster() const override { return false; }
 
 	/** Callback for when a media event occurred in the player. */
-	void HandlePlayerMediaEvent(EMediaEvent Event);
+	MEDIAASSETS_API void HandlePlayerMediaEvent(EMediaEvent Event);
 
 	/** Sets the playlist and properly handles cases when this MediaPlayer object is in disregard for GC set */
-	void SetPlaylistInternal(UMediaPlaylist* InPlaylist) const;
+	MEDIAASSETS_API void SetPlaylistInternal(UMediaPlaylist* InPlaylist) const;
 
 	/** Open media source with the given options. */
-	bool OpenSourceInternal(UMediaSource* MediaSource, const FMediaPlayerOptions* Options);
+	MEDIAASSETS_API bool OpenSourceInternal(UMediaSource* MediaSource, const FMediaPlayerOptions* Options);
 
 	/** Ensure internal playlist is created */
-	void EnsurePlaylist() const;
+	MEDIAASSETS_API void EnsurePlaylist() const;
 
 private:
 
diff --git a/Engine/Source/Runtime/MediaAssets/Public/MediaPlayerProxyInterface.h b/Engine/Source/Runtime/MediaAssets/Public/MediaPlayerProxyInterface.h
--- a/Engine/Source/Runtime/MediaAssets/Public/MediaPlayerProxyInterface.h
+++ b/Engine/Source/Runtime/MediaAssets/Public/MediaPlayerProxyInterface.h
@@ -22,7 +22,7 @@ class UMediaPlayerProxyInterface : public UInterface
 	GENERATED_BODY()
 };
 
-class MEDIAASSETS_API IMediaPlayerProxyInterface
+class IMediaPlayerProxyInterface
 {
 	GENERATED_BODY()
 
diff --git a/Engine/Source/Runtime/MediaAssets/Public/MediaPlaylist.h b/Engine/Source/Runtime/MediaAssets/Public/MediaPlaylist.h
--- a/Engine/Source/Runtime/MediaAssets/Public/MediaPlaylist.h
+++ b/Engine/Source/Runtime/MediaAssets/Public/MediaPlaylist.h
@@ -20,8 +20,8 @@ struct FFrame;
 /**
  * Implements a media play list.
  */
-UCLASS(BlueprintType, hidecategories=(Object))
-class MEDIAASSETS_API UMediaPlaylist
+UCLASS(BlueprintType, hidecategories=(Object), MinimalAPI)
+class UMediaPlaylist
 	: public UObject
 {
 	GENERATED_BODY()
@@ -36,7 +36,7 @@ public:
 	 * @see AddFile, AddUrl, Insert, RemoveAll, Remove, Replace
 	 */
 	UFUNCTION(BlueprintCallable, Category="Media|MediaPlaylist")
-	bool Add(UMediaSource* MediaSource);
+	MEDIAASSETS_API bool Add(UMediaSource* MediaSource);
 
 	/**
 	 * Add a media file path to the play list.
@@ -46,7 +46,7 @@ public:
 	 * @see Add, AddUrl, Insert, RemoveAll, Remove, Replace
 	 */
 	UFUNCTION(BlueprintCallable, Category="Media|MediaPlaylist")
-	bool AddFile(const FString& FilePath);
+	MEDIAASSETS_API bool AddFile(const FString& FilePath);
 
 	/**
 	 * Add a media URL to the play list.
@@ -56,7 +56,7 @@ public:
 	 * @see Add, AddFile, Insert, RemoveAll, Remove, Replace
 	 */
 	UFUNCTION(BlueprintCallable, Category="Media|MediaPlaylist")
-	bool AddUrl(const FString& Url);
+	MEDIAASSETS_API bool AddUrl(const FString& Url);
 
 	/**
 	 * Get the media source at the specified index.
@@ -66,7 +66,7 @@ public:
 	 * @see GetNext, GetRandom
 	 */
 	UFUNCTION(BlueprintCallable, Category="Media|MediaPlaylist")
-	UMediaSource* Get(int32 Index);
+	MEDIAASSETS_API UMediaSource* Get(int32 Index);
 
 	/**
 	 * Get the next media source in the play list.
@@ -76,7 +76,7 @@ public:
 	 * @see , GetPrevious, GetRandom
 	 */
 	UFUNCTION(BlueprintCallable, Category="Media|MediaPlaylist")
-	UMediaSource* GetNext(int32& InOutIndex);
+	MEDIAASSETS_API UMediaSource* GetNext(int32& InOutIndex);
 
 	/**
 	 * Get the previous media source in the play list.
@@ -86,7 +86,7 @@ public:
 	 * @see , GetNext, GetRandom
 	 */
 	UFUNCTION(BlueprintCallable, Category="Media|MediaPlaylist")
-	UMediaSource* GetPrevious(int32& InOutIndex);
+	MEDIAASSETS_API UMediaSource* GetPrevious(int32& InOutIndex);
 
 	/**
 	 * Get a random media source in the play list.
@@ -96,7 +96,7 @@ public:
 	 * @see Get, GetNext, GetPrevious
 	 */
 	UFUNCTION(BlueprintCallable, Category="Media|MediaPlaylist")
-	UMediaSource* GetRandom(int32& OutIndex);
+	MEDIAASSETS_API UMediaSource* GetRandom(int32& OutIndex);
 
 	/**
 	 * Insert a media source into the play list at the given position.
@@ -106,7 +106,7 @@ public:
 	 * @see Add, Remove, RemoveAll, Replace
 	 */
 	UFUNCTION(BlueprintCallable, Category="Media|MediaPlaylist")
-	void Insert(UMediaSource* MediaSource, int32 Index);
+	MEDIAASSETS_API void Insert(UMediaSource* MediaSource, int32 Index);
 
 	/**
 	 * Get the number of media sources in the play list.
@@ -127,7 +127,7 @@ public:
 	 * @see Add, Insert, Remove, Replace
 	 */
 	UFUNCTION(BlueprintCallable, Category="Media|MediaPlaylist")
-	bool Remove(UMediaSource* MediaSource);
+	MEDIAASSETS_API bool Remove(UMediaSource* MediaSource);
 
 	/**
 	 * Remove the media source at the specified position.
@@ -137,7 +137,7 @@ public:
 	 * @see Add, Insert, RemoveAll, Replace
 	 */
 	UFUNCTION(BlueprintCallable, Category="Media|MediaPlaylist")
-	bool RemoveAt(int32 Index);
+	MEDIAASSETS_API bool RemoveAt(int32 Index);
 
 	/**
 	 * Replace the media source at the specified position.
@@ -148,7 +148,7 @@ public:
 	 * @see Add, Insert, RemoveAll, RemoveAt
 	 */
 	UFUNCTION(BlueprintCallable, Category="Media|MediaPlaylist")
-	bool Replace(int32 Index, UMediaSource* Replacement);
+	MEDIAASSETS_API bool Replace(int32 Index, UMediaSource* Replacement);
 
 protected:
 
diff --git a/Engine/Source/Runtime/MediaAssets/Public/MediaSoundComponent.h b/Engine/Source/Runtime/MediaAssets/Public/MediaSoundComponent.h
--- a/Engine/Source/Runtime/MediaAssets/Public/MediaSoundComponent.h
+++ b/Engine/Source/Runtime/MediaAssets/Public/MediaSoundComponent.h
@@ -147,8 +147,8 @@ private:
 /**
  * Implements a sound component for playing a media player's audio output.
  */
-UCLASS(ClassGroup=Media, editinlinenew, meta=(BlueprintSpawnableComponent))
-class MEDIAASSETS_API UMediaSoundComponent
+UCLASS(ClassGroup=Media, editinlinenew, meta=(BlueprintSpawnableComponent), MinimalAPI)
+class UMediaSoundComponent
 	: public USynthComponent
 {
 	GENERATED_BODY()
@@ -189,10 +189,10 @@ public:
 	 *
 	 * @param ObjectInitializer Initialization parameters.
 	 */
-	UMediaSoundComponent(const FObjectInitializer& ObjectInitializer);
+	MEDIAASSETS_API UMediaSoundComponent(const FObjectInitializer& ObjectInitializer);
 
 	/** Virtual destructor. */
-	~UMediaSoundComponent();
+	MEDIAASSETS_API ~UMediaSoundComponent();
 
 public:
 
@@ -203,7 +203,7 @@ public:
 	 * @return true if attenuation settings were returned, false if attenuation is disabled.
 	 */
 	UFUNCTION(BlueprintCallable, Category="Media|MediaSoundComponent", meta=(DisplayName="Get Attenuation Settings To Apply", ScriptName="GetAttenuationSettingsToApply"))
-	bool BP_GetAttenuationSettingsToApply(FSoundAttenuationSettings& OutAttenuationSettings);
+	MEDIAASSETS_API bool BP_GetAttenuationSettingsToApply(FSoundAttenuationSettings& OutAttenuationSettings);
 
 	/**
 	 * Get the media player that provides the audio samples.
@@ -212,7 +212,7 @@ public:
 	 * @see SetMediaPlayer
 	 */
 	UFUNCTION(BlueprintCallable, Category="Media|MediaSoundComponent")
-	UMediaPlayer* GetMediaPlayer() const;
+	MEDIAASSETS_API UMediaPlayer* GetMediaPlayer() const;
 
 	virtual USoundClass* GetSoundClass() override
 	{
@@ -244,45 +244,45 @@ public:
 	 * @see GetMediaPlayer
 	 */
 	UFUNCTION(BlueprintCallable, Category="Media|MediaSoundComponent")
-	void SetMediaPlayer(UMediaPlayer* NewMediaPlayer);
+	MEDIAASSETS_API void SetMediaPlayer(UMediaPlayer* NewMediaPlayer);
 
 	/** Turns on spectral analysis of the audio generated in the media sound component. */
 	UFUNCTION(BlueprintCallable, Category = "Media|MediaSoundComponent")
-	void SetEnableSpectralAnalysis(bool bInSpectralAnalysisEnabled);
+	MEDIAASSETS_API void SetEnableSpectralAnalysis(bool bInSpectralAnalysisEnabled);
 	
 	/** Sets the settings to use for spectral analysis. */
 	UFUNCTION(BlueprintCallable, Category = "Media|MediaSoundComponent")
-	void SetSpectralAnalysisSettings(TArray<float> InFrequenciesToAnalyze, EMediaSoundComponentFFTSize InFFTSize = EMediaSoundComponentFFTSize::Medium_512);
+	MEDIAASSETS_API void SetSpectralAnalysisSettings(TArray<float> InFrequenciesToAnalyze, EMediaSoundComponentFFTSize InFFTSize = EMediaSoundComponentFFTSize::Medium_512);
 
 	/** Retrieves the spectral data if spectral analysis is enabled. */
 	UFUNCTION(BlueprintCallable, Category = "TimeSynth")
-	TArray<FMediaSoundComponentSpectralData> GetSpectralData();
+	MEDIAASSETS_API TArray<FMediaSoundComponentSpectralData> GetSpectralData();
 
 	/** Retrieves and normalizes the spectral data if spectral analysis is enabled. */
 	UFUNCTION(BlueprintCallable, Category = "TimeSynth")
-	TArray<FMediaSoundComponentSpectralData> GetNormalizedSpectralData();
+	MEDIAASSETS_API TArray<FMediaSoundComponentSpectralData> GetNormalizedSpectralData();
 
 	/** Turns on amplitude envelope following the audio in the media sound component. */
 	UFUNCTION(BlueprintCallable, Category = "Media|MediaSoundComponent")
-	void SetEnableEnvelopeFollowing(bool bInEnvelopeFollowing);
+	MEDIAASSETS_API void SetEnableEnvelopeFollowing(bool bInEnvelopeFollowing);
 
 	/** Sets the envelope attack and release times (in ms). */
 	UFUNCTION(BlueprintCallable, Category = "Media|MediaSoundComponent")
-	void SetEnvelopeFollowingsettings(int32 AttackTimeMsec, int32 ReleaseTimeMsec);
+	MEDIAASSETS_API void SetEnvelopeFollowingsettings(int32 AttackTimeMsec, int32 ReleaseTimeMsec);
 
 	/** Retrieves the current amplitude envelope. */
 	UFUNCTION(BlueprintCallable, Category = "TimeSynth")
-	float GetEnvelopeValue() const;
+	MEDIAASSETS_API float GetEnvelopeValue() const;
 
 public:
 
 	/** Adds a clock sink so this can be ticked without the world. */
-	void AddClockSink();
+	MEDIAASSETS_API void AddClockSink();
 
 	/** Removes the clock sink. */
-	void RemoveClockSink();
+	MEDIAASSETS_API void RemoveClockSink();
 
-	void UpdatePlayer();
+	MEDIAASSETS_API void UpdatePlayer();
 
 #if WITH_EDITOR
 	/**
@@ -291,37 +291,37 @@ public:
 	 * @param NewMediaPlayer The player to set.
 	 * @see SetMediaPlayer
 	 */
-	void SetDefaultMediaPlayer(UMediaPlayer* NewMediaPlayer);
+	MEDIAASSETS_API void SetDefaultMediaPlayer(UMediaPlayer* NewMediaPlayer);
 #endif
 
 public:
 
 	//~ TAttenuatedComponentVisualizer interface
 
-	void CollectAttenuationShapesForVisualization(TMultiMap<EAttenuationShape::Type, FBaseAttenuationSettings::AttenuationShapeDetails>& ShapeDetailsMap) const;
+	MEDIAASSETS_API void CollectAttenuationShapesForVisualization(TMultiMap<EAttenuationShape::Type, FBaseAttenuationSettings::AttenuationShapeDetails>& ShapeDetailsMap) const;
 
 public:
 
 	//~ UActorComponent interface
 
-	virtual void OnRegister() override;
-	virtual void OnUnregister() override;
-	virtual void TickComponent(float DeltaTime, enum ELevelTick TickType, FActorComponentTickFunction *ThisTickFunction) override;
+	MEDIAASSETS_API virtual void OnRegister() override;
+	MEDIAASSETS_API virtual void OnUnregister() override;
+	MEDIAASSETS_API virtual void TickComponent(float DeltaTime, enum ELevelTick TickType, FActorComponentTickFunction *ThisTickFunction) override;
 
 public:
 
 	//~ USceneComponent interface
 
-	virtual void Activate(bool bReset = false) override;
-	virtual void Deactivate() override;
+	MEDIAASSETS_API virtual void Activate(bool bReset = false) override;
+	MEDIAASSETS_API virtual void Deactivate() override;
 
 public:
 
 	//~ UObject interface
-	virtual void PostLoad() override;
+	MEDIAASSETS_API virtual void PostLoad() override;
 
 #if WITH_EDITOR
-	virtual void PostEditChangeProperty(FPropertyChangedEvent& PropertyChangedEvent) override;
+	MEDIAASSETS_API virtual void PostEditChangeProperty(FPropertyChangedEvent& PropertyChangedEvent) override;
 #endif
 
 protected:
@@ -331,15 +331,15 @@ protected:
 	 *
 	 * @return Attenuation settings, or nullptr if attenuation is disabled.
 	 */
-	const FSoundAttenuationSettings* GetSelectedAttenuationSettings() const;
+	MEDIAASSETS_API const FSoundAttenuationSettings* GetSelectedAttenuationSettings() const;
 
 protected:
 
 	//~ USynthComponent interface
 
-	virtual bool Init(int32& SampleRate) override;
+	MEDIAASSETS_API virtual bool Init(int32& SampleRate) override;
 
-	virtual ISoundGeneratorPtr CreateSoundGenerator(const FSoundGeneratorInitParams& InParams) override;
+	MEDIAASSETS_API virtual ISoundGeneratorPtr CreateSoundGenerator(const FSoundGeneratorInitParams& InParams) override;
 
 protected:
 
diff --git a/Engine/Source/Runtime/MediaAssets/Public/MediaSource.h b/Engine/Source/Runtime/MediaAssets/Public/MediaSource.h
--- a/Engine/Source/Runtime/MediaAssets/Public/MediaSource.h
+++ b/Engine/Source/Runtime/MediaAssets/Public/MediaSource.h
@@ -60,8 +60,8 @@ struct FMediaSourceCacheSettings
  * location is encoded as a media URL string, whose URI scheme and optional file
  * extension will be used to locate a suitable media player.
  */
-UCLASS(Abstract, editinlinenew, BlueprintType, hidecategories=(Object))
-class MEDIAASSETS_API UMediaSource
+UCLASS(Abstract, editinlinenew, BlueprintType, hidecategories=(Object), MinimalAPI)
+class UMediaSource
 	: public UObject
 	, public IMediaOptions
 {
@@ -76,7 +76,7 @@ public:
 	 * @see GetProxies
 	 */
 	UFUNCTION(BlueprintCallable, Category="Media|MediaSource")
-	virtual FString GetUrl() const PURE_VIRTUAL(UMediaSource::GetUrl, return FString(););
+	MEDIAASSETS_API virtual FString GetUrl() const PURE_VIRTUAL(UMediaSource::GetUrl, return FString(););
 
 	/**
 	 * Validate the media source settings (must be implemented in child classes).
@@ -84,12 +84,12 @@ public:
 	 * @return true if validation passed, false otherwise.
 	 */
 	UFUNCTION(BlueprintCallable, Category="Media|MediaSource")
-	virtual bool Validate() const PURE_VIRTUAL(UMediaSource::Validate, return false;);
+	MEDIAASSETS_API virtual bool Validate() const PURE_VIRTUAL(UMediaSource::Validate, return false;);
 
 	/**
 	 * Call this to set cache settings to pass to the player.
 	 */
-	void SetCacheSettings(const FMediaSourceCacheSettings& Settings);
+	MEDIAASSETS_API void SetCacheSettings(const FMediaSourceCacheSettings& Settings);
 
 
 #if WITH_EDITOR
@@ -97,7 +97,7 @@ public:
 	/**
 	 * Starts the process to generate a thumbnail.
 	 */
-	void GenerateThumbnail();
+	MEDIAASSETS_API void GenerateThumbnail();
 
 	/**
 	 * Gets our thumbnail texture, if any.
@@ -119,14 +119,14 @@ public:
 	 * @param InDelegate	This will get called if the Url passed into GetMediaSourceForUrl
 	 *						matches Extension.
 	 */
-	static void RegisterSpawnFromFileExtension(const FString& Extension, FMediaSourceSpawnDelegate InDelegate);
+	static MEDIAASSETS_API void RegisterSpawnFromFileExtension(const FString& Extension, FMediaSourceSpawnDelegate InDelegate);
 	
 	/**
 	 * Call this to unregister a callback set with RegisterSpawnFromFileExtension.
 	 *
 	 * @param Extension		File extension that the callack was registered with.
 	 */
-	static void UnregisterSpawnFromFileExtension(const FString& Extension);
+	static MEDIAASSETS_API void UnregisterSpawnFromFileExtension(const FString& Extension);
 
 	/**
 	 * Call this to try and create a media source appropriate for the media.
@@ -135,37 +135,37 @@ public:
 	 * @param Outer			Outer to use for this object.
 	 * @return				Media source or nullptr if none are appropriate.
 	 */
-	static UMediaSource* SpawnMediaSourceForString(const FString& MediaPath, UObject* Outer);
+	static MEDIAASSETS_API UMediaSource* SpawnMediaSourceForString(const FString& MediaPath, UObject* Outer);
 
 public:
 	//~ UObject interface
-	virtual void BeginDestroy() override;
+	MEDIAASSETS_API virtual void BeginDestroy() override;
 
 	//~ IMediaOptions interface
 
-	virtual FName GetDesiredPlayerName() const override;
-	virtual bool GetMediaOption(const FName& Key, bool DefaultValue) const override;
-	virtual double GetMediaOption(const FName& Key, double DefaultValue) const override;
-	virtual int64 GetMediaOption(const FName& Key, int64 DefaultValue) const override;
-	virtual FString GetMediaOption(const FName& Key, const FString& DefaultValue) const override;
-	virtual FText GetMediaOption(const FName& Key, const FText& DefaultValue) const override;
-	virtual TSharedPtr<FDataContainer, ESPMode::ThreadSafe> GetMediaOption(const FName& Key, const TSharedPtr<FDataContainer, ESPMode::ThreadSafe>& DefaultValue) const override;
-	virtual bool HasMediaOption(const FName& Key) const override;
+	MEDIAASSETS_API virtual FName GetDesiredPlayerName() const override;
+	MEDIAASSETS_API virtual bool GetMediaOption(const FName& Key, bool DefaultValue) const override;
+	MEDIAASSETS_API virtual double GetMediaOption(const FName& Key, double DefaultValue) const override;
+	MEDIAASSETS_API virtual int64 GetMediaOption(const FName& Key, int64 DefaultValue) const override;
+	MEDIAASSETS_API virtual FString GetMediaOption(const FName& Key, const FString& DefaultValue) const override;
+	MEDIAASSETS_API virtual FText GetMediaOption(const FName& Key, const FText& DefaultValue) const override;
+	MEDIAASSETS_API virtual TSharedPtr<FDataContainer, ESPMode::ThreadSafe> GetMediaOption(const FName& Key, const TSharedPtr<FDataContainer, ESPMode::ThreadSafe>& DefaultValue) const override;
+	MEDIAASSETS_API virtual bool HasMediaOption(const FName& Key) const override;
 
 	/** Set a boolean parameter to pass to the player. */
 	UFUNCTION(BlueprintCallable, meta=(DisplayName = "SetMediaOption (boolean)"), Category = "Media|MediaSource")
-	void SetMediaOptionBool(const FName& Key, bool Value);
+	MEDIAASSETS_API void SetMediaOptionBool(const FName& Key, bool Value);
 	/** Set a float parameter to pass to the player. */
 	UFUNCTION(BlueprintCallable, meta = (DisplayName = "SetMediaOption (float)"), Category = "Media|MediaSource")
-	void SetMediaOptionFloat(const FName& Key, float Value);
+	MEDIAASSETS_API void SetMediaOptionFloat(const FName& Key, float Value);
 	/** Set a double parameter to pass to the player. */
-	void SetMediaOptionDouble(const FName& Key, double Value);
+	MEDIAASSETS_API void SetMediaOptionDouble(const FName& Key, double Value);
 	/** Set an integer64 parameter to pass to the player. */
 	UFUNCTION(BlueprintCallable, meta = (DisplayName = "SetMediaOption (integer64)"), Category = "Media|MediaSource")
-	void SetMediaOptionInt64(const FName& Key, int64 Value);
+	MEDIAASSETS_API void SetMediaOptionInt64(const FName& Key, int64 Value);
 	/** Set a string parameter to pass to the player. */
 	UFUNCTION(BlueprintCallable, meta = (DisplayName = "SetMediaOption (string)"), Category = "Media|MediaSource")
-	void SetMediaOptionString(const FName& Key, const FString& Value);
+	MEDIAASSETS_API void SetMediaOptionString(const FName& Key, const FString& Value);
 
 private:
 	/** Holds our media options. */
@@ -175,17 +175,17 @@ private:
 	 * Get the media option specified by the Key as a Variant.
 	 * Returns nullptr if the Key does not exist.
 	 */
-	const FVariant* GetMediaOptionDefault(const FName& Key) const;
+	MEDIAASSETS_API const FVariant* GetMediaOptionDefault(const FName& Key) const;
 
 	/**
 	 * Sets the media option specified by Key to the supplied Variant.
 	 */
-	void SetMediaOption(const FName& Key, FVariant& Value);
+	MEDIAASSETS_API void SetMediaOption(const FName& Key, FVariant& Value);
 
 	/**
 	 * Get a mapping of file extensions to spawn delegates.
 	 */
-	static TMap<FString, FMediaSourceSpawnDelegate>& GetSpawnFromFileExtensionDelegates();
+	static MEDIAASSETS_API TMap<FString, FMediaSourceSpawnDelegate>& GetSpawnFromFileExtensionDelegates();
 
 #if WITH_EDITORONLY_DATA
 
diff --git a/Engine/Source/Runtime/MediaAssets/Public/MediaSourceRendererInterface.h b/Engine/Source/Runtime/MediaAssets/Public/MediaSourceRendererInterface.h
--- a/Engine/Source/Runtime/MediaAssets/Public/MediaSourceRendererInterface.h
+++ b/Engine/Source/Runtime/MediaAssets/Public/MediaSourceRendererInterface.h
@@ -18,7 +18,7 @@ class UMediaSourceRendererInterface : public UInterface
 	GENERATED_BODY()
 };
 
-class MEDIAASSETS_API IMediaSourceRendererInterface
+class IMediaSourceRendererInterface
 {
 	GENERATED_BODY()
 
diff --git a/Engine/Source/Runtime/MediaAssets/Public/MediaTexture.h b/Engine/Source/Runtime/MediaAssets/Public/MediaTexture.h
--- a/Engine/Source/Runtime/MediaAssets/Public/MediaTexture.h
+++ b/Engine/Source/Runtime/MediaAssets/Public/MediaTexture.h
@@ -44,8 +44,8 @@ enum MediaTextureOrientation : int
  * note: derives directly from UTexture, not from UTexture2D or UTexture2DDynamic
  *    maybe should have been UTexture2DDynamic?
  */
-UCLASS(hidecategories=(Adjustments, Compositing, LevelOfDetail, ImportSettings, Object))
-class MEDIAASSETS_API UMediaTexture
+UCLASS(hidecategories=(Adjustments, Compositing, LevelOfDetail, ImportSettings, Object), MinimalAPI)
+class UMediaTexture
 	: public UTexture
 {
 	GENERATED_UCLASS_BODY()
@@ -109,7 +109,7 @@ public:
 	 * @see GetHeight, GetWidth
 	 */
 	UFUNCTION(BlueprintCallable, Category="Media|MediaTexture")
-	float GetAspectRatio() const;
+	MEDIAASSETS_API float GetAspectRatio() const;
 
 	/**
 	 * Gets the current height of the texture.
@@ -118,7 +118,7 @@ public:
 	 * @see GetAspectRatio, GetWidth
 	 */
 	UFUNCTION(BlueprintCallable, Category="Media|MediaTexture")
-	int32 GetHeight() const;
+	MEDIAASSETS_API int32 GetHeight() const;
 
 	/**
 	 * Gets the current width of the texture.
@@ -127,7 +127,7 @@ public:
 	 * @see GetAspectRatio, GetHeight
 	 */
 	UFUNCTION(BlueprintCallable, Category="Media|MediaTexture")
-	int32 GetWidth() const;
+	MEDIAASSETS_API int32 GetWidth() const;
 
 	/**
 	 * Gets the current numbe of mips of the texture.
@@ -135,7 +135,7 @@ public:
 	 * @return Number of mips.
 	 */
 	UFUNCTION(BlueprintCallable, Category = "Media|MediaTexture")
-	int32 GetTextureNumMips() const;
+	MEDIAASSETS_API int32 GetTextureNumMips() const;
 
 	/**
 	 * Get the media player that provides the video samples.
@@ -144,7 +144,7 @@ public:
 	 * @see SetMediaPlayer
 	 */
 	UFUNCTION(BlueprintCallable, Category = "Media|MediaTexture")
-	UMediaPlayer* GetMediaPlayer() const;
+	MEDIAASSETS_API UMediaPlayer* GetMediaPlayer() const;
 
 	/**
 	 * Set the media player that provides the video samples.
@@ -153,7 +153,7 @@ public:
 	 * @see GetMediaPlayer
 	 */
 	UFUNCTION(BlueprintCallable, Category="Media|MediaTexture")
-	void SetMediaPlayer(UMediaPlayer* NewMediaPlayer);
+	MEDIAASSETS_API void SetMediaPlayer(UMediaPlayer* NewMediaPlayer);
 	
 	/**
 	 * Creates a new resource for the texture, and updates any cached references to the resource.
@@ -166,7 +166,7 @@ public:
 	 * Caches the next available sample time from the queue when last rendering was made
 	 * @see GetNextSampleTime
 	 */
-	void CacheNextAvailableSampleTime(FTimespan InNextSampleTime);
+	MEDIAASSETS_API void CacheNextAvailableSampleTime(FTimespan InNextSampleTime);
 
 	/**
 	 * Gets the next sample Time. Only valid if GetAvailableSampleCount is greater than 0
@@ -174,13 +174,13 @@ public:
 	 * @return FTimespan of the next sample or FTimespan::MinValue if no sample was available in the queue.
 	 * @see GetAvailableSampleCount, CacheNextAvailableSampleTime
 	 */
-	FTimespan GetNextSampleTime() const;
+	MEDIAASSETS_API FTimespan GetNextSampleTime() const;
 
 	/**
 	 * Gets the number of samples that are available
 	 * @return The number of samples in the queue
 	 */
-	int32 GetAvailableSampleCount() const;
+	MEDIAASSETS_API int32 GetAvailableSampleCount() const;
 
 #if WITH_EDITOR
 
@@ -190,7 +190,7 @@ public:
 	 * @param NewMediaPlayer The player to set.
 	 * @see SetMediaPlayer
 	 */
-	void SetDefaultMediaPlayer(UMediaPlayer* NewMediaPlayer);
+	MEDIAASSETS_API void SetDefaultMediaPlayer(UMediaPlayer* NewMediaPlayer);
 
 #endif
 
@@ -198,19 +198,19 @@ public:
 	 * Get current aspect ratio of presented frame.
 	 * @return Aspect ratio of current frame
 	 */
-	float GetCurrentAspectRatio() const;
+	MEDIAASSETS_API float GetCurrentAspectRatio() const;
 
 	/**
 	 * Get current orientation of presented frame.
 	 * @return Orientation of current frame
 	 */
-	MediaTextureOrientation GetCurrentOrientation() const;
+	MEDIAASSETS_API MediaTextureOrientation GetCurrentOrientation() const;
 
 	/**
 	 * Get the texture's mip-map bias, clamped to a legal range.
 	 * @return Mip-map bias value
 	 */
-	float GetMipMapBias() const;
+	MEDIAASSETS_API float GetMipMapBias() const;
 
 	/**
 	 * Set texture's mip-map bias, for use by the texture resource sampler.
@@ -218,7 +218,7 @@ public:
 	 *
 	 * @param InMipMapBias Bias value.
 	 */
-	void SetMipMapBias(float InMipMapBias);
+	MEDIAASSETS_API void SetMipMapBias(float InMipMapBias);
 
 	/**
 	 * Set an override for color space conversions
@@ -229,7 +229,7 @@ public:
 	}
 
 	/** Renders this media texture. Only has an effect if its RenderMode is ERenderMode::JustInTime */
-	virtual void JustInTimeRender();
+	MEDIAASSETS_API virtual void JustInTimeRender();
 
 	/** Sets the ERenderMode of this media texture */
 	void SetRenderMode(ERenderMode InRenderMode)
@@ -247,29 +247,29 @@ public:
 
 	//~ UTexture interface.
 
-	virtual void BeginDestroy() override;
-	virtual FTextureResource* CreateResource() override;
-	virtual EMaterialValueType GetMaterialType() const override;
-	virtual float GetSurfaceWidth() const override;
-	virtual float GetSurfaceHeight() const override;
+	MEDIAASSETS_API virtual void BeginDestroy() override;
+	MEDIAASSETS_API virtual FTextureResource* CreateResource() override;
+	MEDIAASSETS_API virtual EMaterialValueType GetMaterialType() const override;
+	MEDIAASSETS_API virtual float GetSurfaceWidth() const override;
+	MEDIAASSETS_API virtual float GetSurfaceHeight() const override;
 	virtual float GetSurfaceDepth() const override { return 0; }
 	virtual uint32 GetSurfaceArraySize() const override { return 0; }
-	virtual FGuid GetExternalTextureGuid() const override;
-	void SetRenderedExternalTextureGuid(const FGuid& InNewGuid);
+	MEDIAASSETS_API virtual FGuid GetExternalTextureGuid() const override;
+	MEDIAASSETS_API void SetRenderedExternalTextureGuid(const FGuid& InNewGuid);
 	virtual ETextureClass GetTextureClass() const override { return ETextureClass::Other2DNoSource; }
-	virtual uint32 CalcTextureMemorySizeEnum(ETextureMipCount Enum) const override;
+	MEDIAASSETS_API virtual uint32 CalcTextureMemorySizeEnum(ETextureMipCount Enum) const override;
 
 public:
 
 	//~ UObject interface.
 
-	virtual FString GetDesc() override;
-	virtual void GetResourceSizeEx(FResourceSizeEx& CumulativeResourceSize) override;
-	virtual void PostLoad() override;
-	virtual bool IsPostLoadThreadSafe() const override;
+	MEDIAASSETS_API virtual FString GetDesc() override;
+	MEDIAASSETS_API virtual void GetResourceSizeEx(FResourceSizeEx& CumulativeResourceSize) override;
+	MEDIAASSETS_API virtual void PostLoad() override;
+	MEDIAASSETS_API virtual bool IsPostLoadThreadSafe() const override;
 
 #if WITH_EDITOR
-	virtual void PostEditChangeProperty(FPropertyChangedEvent& PropertyChangedEvent) override;
+	MEDIAASSETS_API virtual void PostEditChangeProperty(FPropertyChangedEvent& PropertyChangedEvent) override;
 #endif
 
 protected:
@@ -279,13 +279,13 @@ protected:
 	 *
 	 * @param Timecode The current timecode.
 	 */
-	void TickResource(FTimespan Timecode);
+	MEDIAASSETS_API void TickResource(FTimespan Timecode);
 
 	/** Update the video sample queue, if necessary. */
-	void UpdatePlayerAndQueue();
+	MEDIAASSETS_API void UpdatePlayerAndQueue();
 
 	/** Update sample info */
-	void UpdateSampleInfo(const TSharedPtr<IMediaTextureSample, ESPMode::ThreadSafe>& Sample);
+	MEDIAASSETS_API void UpdateSampleInfo(const TSharedPtr<IMediaTextureSample, ESPMode::ThreadSafe>& Sample);
 
 protected:
 
diff --git a/Engine/Source/Runtime/MediaAssets/Public/MediaTextureTracker.h b/Engine/Source/Runtime/MediaAssets/Public/MediaTextureTracker.h
--- a/Engine/Source/Runtime/MediaAssets/Public/MediaTextureTracker.h
+++ b/Engine/Source/Runtime/MediaAssets/Public/MediaTextureTracker.h
@@ -80,7 +80,7 @@ public:
 	 *
 	 * @ return List of textures.
 	 */
-	MEDIAASSETS_API const TArray<TWeakObjectPtr<UMediaTexture>>& GetTextures() { return MediaTextures; }
+	const TArray<TWeakObjectPtr<UMediaTexture>>& GetTextures() { return MediaTextures; }
 
 protected:
 	/** Maps a media texture to an array of playback components. */
diff --git a/Engine/Source/Runtime/MediaAssets/Public/Misc/MediaBlueprintFunctionLibrary.h b/Engine/Source/Runtime/MediaAssets/Public/Misc/MediaBlueprintFunctionLibrary.h
--- a/Engine/Source/Runtime/MediaAssets/Public/Misc/MediaBlueprintFunctionLibrary.h
+++ b/Engine/Source/Runtime/MediaAssets/Public/Misc/MediaBlueprintFunctionLibrary.h
@@ -93,7 +93,7 @@ ENUM_CLASS_FLAGS(EMediaWebcamCaptureDeviceFilter)
  * Information about a capture device.
  */
 USTRUCT(BlueprintType)
-struct MEDIAASSETS_API FMediaCaptureDevice
+struct FMediaCaptureDevice
 {
 	GENERATED_USTRUCT_BODY()
 
@@ -119,8 +119,8 @@ struct MEDIAASSETS_API FMediaCaptureDevice
 /**
  * Blueprint library for Media related functions.
  */
-UCLASS(meta=(ScriptName="MediaLibrary"))
-class MEDIAASSETS_API UMediaBlueprintFunctionLibrary
+UCLASS(meta=(ScriptName="MediaLibrary"), MinimalAPI)
+class UMediaBlueprintFunctionLibrary
 	: public UBlueprintFunctionLibrary
 {
 	GENERATED_BODY()
@@ -137,7 +137,7 @@ public:
 	 * @param Filter The types of capture devices to return (-1 = all).
 	 */
 	UFUNCTION(BlueprintCallable, Category = "Media|Capture")
-	static void EnumerateAudioCaptureDevices(TArray<FMediaCaptureDevice>& OutDevices, int32 Filter = -1);
+	static MEDIAASSETS_API void EnumerateAudioCaptureDevices(TArray<FMediaCaptureDevice>& OutDevices, int32 Filter = -1);
 
 	/**
 	 * Enumerate available audio capture devices.
@@ -149,7 +149,7 @@ public:
 	 * @param Filter The types of capture devices to return (-1 = all).
 	 */
 	UFUNCTION(BlueprintCallable, Category = "Media|Capture")
-	static void EnumerateVideoCaptureDevices(TArray<FMediaCaptureDevice>& OutDevices, int32 Filter = -1);
+	static MEDIAASSETS_API void EnumerateVideoCaptureDevices(TArray<FMediaCaptureDevice>& OutDevices, int32 Filter = -1);
 
 	/**
 	 * Enumerate available audio capture devices.
@@ -161,5 +161,5 @@ public:
 	 * @param Filter The types of capture devices to return (-1 = all).
 	 */
 	UFUNCTION(BlueprintCallable, Category = "Media|Capture")
-	static void EnumerateWebcamCaptureDevices(TArray<FMediaCaptureDevice>& OutDevices, int32 Filter = -1);
+	static MEDIAASSETS_API void EnumerateWebcamCaptureDevices(TArray<FMediaCaptureDevice>& OutDevices, int32 Filter = -1);
 };
diff --git a/Engine/Source/Runtime/MediaAssets/Public/PlatformMediaSource.h b/Engine/Source/Runtime/MediaAssets/Public/PlatformMediaSource.h
--- a/Engine/Source/Runtime/MediaAssets/Public/PlatformMediaSource.h
+++ b/Engine/Source/Runtime/MediaAssets/Public/PlatformMediaSource.h
@@ -25,8 +25,8 @@ struct FGuid;
  *
  * Use this asset to override media sources on a per-platform basis.
  */
-UCLASS(BlueprintType)
-class MEDIAASSETS_API UPlatformMediaSource
+UCLASS(BlueprintType, MinimalAPI)
+class UPlatformMediaSource
 	: public UMediaSource
 {
 	GENERATED_BODY()
@@ -49,27 +49,27 @@ public:
 	//~ UObject interface
 	PRAGMA_DISABLE_DEPRECATION_WARNINGS // Suppress compiler warning on override of deprecated function
 	UE_DEPRECATED(5.0, "Use version that takes FObjectPreSaveContext instead.")
-	virtual void PreSave(const class ITargetPlatform* TargetPlatform);
-	PRAGMA_ENABLE_DEPRECATION_WARNINGS
+	MEDIAASSETS_API virtual void PreSave(const class ITargetPlatform* TargetPlatform);
+	MEDIAASSETS_API PRAGMA_ENABLE_DEPRECATION_WARNINGS
 	virtual void PreSave(FObjectPreSaveContext ObjectSaveContext);
-	virtual void Serialize(FArchive& Ar) override;
+	MEDIAASSETS_API virtual void Serialize(FArchive& Ar) override;
 
 	//~ UMediaSource interface
 
-	virtual FString GetUrl() const override;
-	virtual bool Validate() const override;
+	MEDIAASSETS_API virtual FString GetUrl() const override;
+	MEDIAASSETS_API virtual bool Validate() const override;
 
 public:
 
 	//~ IMediaOptions interface
 
-	virtual FName GetDesiredPlayerName() const override;
-	virtual bool GetMediaOption(const FName& Key, bool DefaultValue) const override;
-	virtual double GetMediaOption(const FName& Key, double DefaultValue) const override;
-	virtual int64 GetMediaOption(const FName& Key, int64 DefaultValue) const override;
-	virtual FString GetMediaOption(const FName& Key, const FString& DefaultValue) const override;
-	virtual FText GetMediaOption(const FName& Key, const FText& DefaultValue) const override;
-	virtual bool HasMediaOption(const FName& Key) const override;
+	MEDIAASSETS_API virtual FName GetDesiredPlayerName() const override;
+	MEDIAASSETS_API virtual bool GetMediaOption(const FName& Key, bool DefaultValue) const override;
+	MEDIAASSETS_API virtual double GetMediaOption(const FName& Key, double DefaultValue) const override;
+	MEDIAASSETS_API virtual int64 GetMediaOption(const FName& Key, int64 DefaultValue) const override;
+	MEDIAASSETS_API virtual FString GetMediaOption(const FName& Key, const FString& DefaultValue) const override;
+	MEDIAASSETS_API virtual FText GetMediaOption(const FName& Key, const FText& DefaultValue) const override;
+	MEDIAASSETS_API virtual bool HasMediaOption(const FName& Key) const override;
 
 private:
 
diff --git a/Engine/Source/Runtime/MediaAssets/Public/StreamMediaSource.h b/Engine/Source/Runtime/MediaAssets/Public/StreamMediaSource.h
--- a/Engine/Source/Runtime/MediaAssets/Public/StreamMediaSource.h
+++ b/Engine/Source/Runtime/MediaAssets/Public/StreamMediaSource.h
@@ -12,8 +12,8 @@
 class UObject;
 
 
-UCLASS(BlueprintType)
-class MEDIAASSETS_API UStreamMediaSource
+UCLASS(BlueprintType, MinimalAPI)
+class UStreamMediaSource
 	: public UBaseMediaSource
 {
 	GENERATED_BODY()
@@ -28,11 +28,11 @@ public:
 
 	//~ UMediaSource interface
 
-	virtual FString GetUrl() const override;
-	virtual bool Validate() const override;
+	MEDIAASSETS_API virtual FString GetUrl() const override;
+	MEDIAASSETS_API virtual bool Validate() const override;
 
 #if WITH_EDITOR
-	virtual void PostEditChangeProperty(FPropertyChangedEvent& PropertyChangedEvent) override;
+	MEDIAASSETS_API virtual void PostEditChangeProperty(FPropertyChangedEvent& PropertyChangedEvent) override;
 #endif
 
 };
diff --git a/Engine/Source/Runtime/MediaAssets/Public/TimeSynchronizableMediaSource.h b/Engine/Source/Runtime/MediaAssets/Public/TimeSynchronizableMediaSource.h
--- a/Engine/Source/Runtime/MediaAssets/Public/TimeSynchronizableMediaSource.h
+++ b/Engine/Source/Runtime/MediaAssets/Public/TimeSynchronizableMediaSource.h
@@ -26,14 +26,14 @@ namespace TimeSynchronizableMedia
 /**
  * Base class for media sources that can be synchronized with the engine's timecode.
  */
-UCLASS(Abstract)
-class MEDIAASSETS_API UTimeSynchronizableMediaSource : public UBaseMediaSource
+UCLASS(Abstract, MinimalAPI)
+class UTimeSynchronizableMediaSource : public UBaseMediaSource
 {
 	GENERATED_BODY()
 	
 public:
 	/** Default constructor. */
-	UTimeSynchronizableMediaSource();
+	MEDIAASSETS_API UTimeSynchronizableMediaSource();
 
 public:
 
@@ -60,11 +60,11 @@ public:
 public:
 	//~ IMediaOptions interface
 	using Super::GetMediaOption;
-	virtual bool GetMediaOption(const FName& Key, bool DefaultValue) const override;
-	virtual int64 GetMediaOption(const FName& Key, int64 DefaultValue) const override;
-	virtual double GetMediaOption(const FName& Key, double DefaultValue) const override;
-	virtual FString GetMediaOption(const FName& Key, const FString& DefaultValue) const override;
-	virtual bool HasMediaOption(const FName& Key) const override;
+	MEDIAASSETS_API virtual bool GetMediaOption(const FName& Key, bool DefaultValue) const override;
+	MEDIAASSETS_API virtual int64 GetMediaOption(const FName& Key, int64 DefaultValue) const override;
+	MEDIAASSETS_API virtual double GetMediaOption(const FName& Key, double DefaultValue) const override;
+	MEDIAASSETS_API virtual FString GetMediaOption(const FName& Key, const FString& DefaultValue) const override;
+	MEDIAASSETS_API virtual bool HasMediaOption(const FName& Key) const override;
 
 	virtual bool SupportsFormatAutoDetection() const { return false; }
 };
diff --git a/Engine/Source/Runtime/MediaUtils/Public/AppMediaTimeSource.h b/Engine/Source/Runtime/MediaUtils/Public/AppMediaTimeSource.h
--- a/Engine/Source/Runtime/MediaUtils/Public/AppMediaTimeSource.h
+++ b/Engine/Source/Runtime/MediaUtils/Public/AppMediaTimeSource.h
@@ -9,7 +9,7 @@
 /**
  * Implements the a media time source that derives its time from the application's global time.
  */
-class MEDIAUTILS_API FAppMediaTimeSource
+class FAppMediaTimeSource
 	: public IMediaTimeSource
 {
 public:
@@ -21,5 +21,5 @@ public:
 
 	//~ IMediaTimeSource interface
 
-	virtual FTimespan GetTimecode() override;
+	MEDIAUTILS_API virtual FTimespan GetTimecode() override;
 };
diff --git a/Engine/Source/Runtime/MediaUtils/Public/MediaAudioResampler.h b/Engine/Source/Runtime/MediaUtils/Public/MediaAudioResampler.h
--- a/Engine/Source/Runtime/MediaUtils/Public/MediaAudioResampler.h
+++ b/Engine/Source/Runtime/MediaUtils/Public/MediaAudioResampler.h
@@ -16,17 +16,17 @@ class IMediaAudioSample;
 /**
  *
  */
-class MEDIAUTILS_API FMediaAudioResampler
+class FMediaAudioResampler
 {
 public:
 
 	/** Default constructor. */
-	FMediaAudioResampler();
+	MEDIAUTILS_API FMediaAudioResampler();
 
 public:
 
 	/** Flush the resampler. */
-	void Flush();
+	MEDIAUTILS_API void Flush();
 
 	/**
 	 * Generate the next frame of audio.
@@ -40,7 +40,7 @@ public:
 	 * @param JumpFrame Frame that a jump in time occured on.
 	 * @return The actual number of frames returned.
 	 */
-	uint32 Generate(float* Output, FMediaTimeStamp& OutTime, const uint32 FramesRequested, float Rate, FTimespan Time, FMediaAudioSampleSource& SampleSource, uint32& JumpFrame);
+	MEDIAUTILS_API uint32 Generate(float* Output, FMediaTimeStamp& OutTime, const uint32 FramesRequested, float Rate, FTimespan Time, FMediaAudioSampleSource& SampleSource, uint32& JumpFrame);
 
 	/**
 	 * Initialize the resampler.
@@ -48,7 +48,7 @@ public:
 	 * @param InOutputSampleRate Desired output sample rate.
 	 * @see SetAudioSamples
 	 */
-	void Initialize(const uint32 InOutputChannels, const uint32 InOutputSampleRate);
+	MEDIAUTILS_API void Initialize(const uint32 InOutputChannels, const uint32 InOutputSampleRate);
 
 protected:
 
@@ -57,7 +57,7 @@ protected:
 	 *
 	 * @see SetInput
 	 */
-	void ClearInput();
+	MEDIAUTILS_API void ClearInput();
 
 	/**
 	 * Set the audio sample to be resampled.
@@ -66,7 +66,7 @@ protected:
 	 * @return true on success, false otherwise.
 	 * @see ClearInput
 	 */
-	bool SetInput(const TSharedPtr<IMediaAudioSample, ESPMode::ThreadSafe>& Sample);
+	MEDIAUTILS_API bool SetInput(const TSharedPtr<IMediaAudioSample, ESPMode::ThreadSafe>& Sample);
 
 private:
 
diff --git a/Engine/Source/Runtime/MediaUtils/Public/MediaDelegates.h b/Engine/Source/Runtime/MediaUtils/Public/MediaDelegates.h
--- a/Engine/Source/Runtime/MediaUtils/Public/MediaDelegates.h
+++ b/Engine/Source/Runtime/MediaUtils/Public/MediaDelegates.h
@@ -12,7 +12,7 @@
 class IMediaTextureSample;
 class UMediaTexture;
 
-class MEDIAUTILS_API FMediaDelegates
+class FMediaDelegates
 {
 public:
 	using TSharedMediaTextureSample = TSharedPtr<IMediaTextureSample, ESPMode::ThreadSafe>;
@@ -28,14 +28,14 @@ public:
 	 * Used to debug the samples queue. Only available in a development build.
 	 * @see UE_MEDIAUTILS_DEVELOPMENT_DELEGATE
 	 */
-	static FOnSampleDiscardedDelegate OnSampleDiscarded_RenderThread;
+	static MEDIAUTILS_API FOnSampleDiscardedDelegate OnSampleDiscarded_RenderThread;
 
 	/**
 	 * Callback when a sample is going to be used to update the MediaTexture.
 	 * Used to debug the samples queue. Only available in a development build.
 	 * @see UE_MEDIAUTILS_DEVELOPMENT_DELEGATE
 	 */
-	static FOnPreSampleRenderDelegate OnPreSampleRender_RenderThread;
+	static MEDIAUTILS_API FOnPreSampleRenderDelegate OnPreSampleRender_RenderThread;
 
 #endif //UE_MEDIAUTILS_DEVELOPMENT_DELEGATE
 };
