Fix SharedMemoryMediaPlayer not working

diff --git a/Engine/Source/Runtime/MediaAssets/Private/Misc/MediaTextureResource.cpp b/Engine/Source/Runtime/MediaAssets/Private/Misc/MediaTextureResource.cpp
--- a/Engine/Source/Runtime/MediaAssets/Private/Misc/MediaTextureResource.cpp
+++ b/Engine/Source/Runtime/MediaAssets/Private/Misc/MediaTextureResource.cpp
@@ -533,11 +533,13 @@ void FMediaTextureResource::Render(const FRenderParams& Params)
 
 			if (IMediaTextureSampleColorConverter* Converter = Sample->GetMediaTextureSampleColorConverter())
 			{
-				if (bRecreateOutputTarget)
+				FTexture2DRHIRef TextureRef = IntermediateTarget ? IntermediateTarget : RenderTargetTextureRHI;
+
+				if (bRecreateOutputTarget && TextureRef)
 				{
 					const static FLazyName ClassName(TEXT("FMediaTextureResource"));
 
-					const FRHITextureDesc& IntermediateTextureDesc = IntermediateTarget->GetDesc();
+					const FRHITextureDesc& IntermediateTextureDesc = TextureRef->GetDesc();
 			
 					const FRHITextureCreateDesc Desc =
 						FRHITextureCreateDesc::Create2D(TEXT("MediaTextureColorConversionOutput"),
@@ -557,14 +559,14 @@ void FMediaTextureResource::Render(const FRenderParams& Params)
 					bRecreateOutputTarget = false;
 				}
 		
-				Converter->ApplyColorConversion(IntermediateTarget, OutputTarget);
+				Converter->ApplyColorConversion(TextureRef, OutputTarget);
 
 				if (RenderTargetTextureRHI != OutputTarget)
 				{
 					UpdateTextureReference(OutputTarget);
 				}
 			}
-			else
+			else if (IntermediateTarget)
 			{
 				OutputTarget = IntermediateTarget;
 				UpdateTextureReference(OutputTarget);
