Media: Enable new style output by default only on new media textures.

diff --git a/Engine/Source/Runtime/MediaAssets/Private/Assets/MediaTexture.cpp b/Engine/Source/Runtime/MediaAssets/Private/Assets/MediaTexture.cpp
--- a/Engine/Source/Runtime/MediaAssets/Private/Assets/MediaTexture.cpp
+++ b/Engine/Source/Runtime/MediaAssets/Private/Assets/MediaTexture.cpp
@@ -275,6 +275,18 @@ void UMediaTexture::GetResourceSizeEx(FResourceSizeEx& CumulativeResourceSize)
 }
 
 
+void UMediaTexture::PostInitProperties()
+{
+	Super::PostInitProperties();
+
+#if WITH_EDITORONLY_DATA
+	if (!HasAnyFlags(RF_ClassDefaultObject | RF_NeedLoad))
+	{
+		NewStyleOutput = true;
+	}
+#endif
+}
+
 void UMediaTexture::PostLoad()
 {
 	Super::PostLoad();
diff --git a/Engine/Source/Runtime/MediaAssets/Public/MediaTexture.h b/Engine/Source/Runtime/MediaAssets/Public/MediaTexture.h
--- a/Engine/Source/Runtime/MediaAssets/Public/MediaTexture.h
+++ b/Engine/Source/Runtime/MediaAssets/Public/MediaTexture.h
@@ -85,7 +85,7 @@ public:
 	UPROPERTY(EditAnywhere, BlueprintReadWrite, Category = "MediaTexture", meta=(DisplayName="Total number of Mipmaps to output"))
 	uint8 NumMips;
 
-	/** Enable new style output (default = false). */
+	/** Enable new style output. */
 	UPROPERTY(EditAnywhere, BlueprintReadWrite, Category = "MediaTexture", meta = (DisplayName = "Enable new style output"))
 	bool NewStyleOutput;
 
@@ -265,6 +265,7 @@ public:
 
 	MEDIAASSETS_API virtual FString GetDesc() override;
 	MEDIAASSETS_API virtual void GetResourceSizeEx(FResourceSizeEx& CumulativeResourceSize) override;
+	MEDIAASSETS_API virtual void PostInitProperties() override;
 	MEDIAASSETS_API virtual void PostLoad() override;
 	MEDIAASSETS_API virtual bool IsPostLoadThreadSafe() const override;
 
