Fixing up Switch playback issues
#rnx

diff --git a/Engine/Source/Runtime/MediaAssets/Private/Misc/MediaTextureResource.cpp b/Engine/Source/Runtime/MediaAssets/Private/Misc/MediaTextureResource.cpp
--- a/Engine/Source/Runtime/MediaAssets/Private/Misc/MediaTextureResource.cpp
+++ b/Engine/Source/Runtime/MediaAssets/Private/Misc/MediaTextureResource.cpp
@@ -569,7 +569,6 @@ void FMediaTextureResource::Render(const FRenderParams& Params)
 			else if (IntermediateTarget)
 			{
 				OutputTarget = IntermediateTarget;
-				UpdateTextureReference(OutputTarget);
 			}
 
 			Rotation = Sample->GetScaleRotation();
