Added zeroing of output buffer when no samples could be provided. This was previously done by the upper layer code but was removed there in CL 27923312.

#rnx

diff --git a/Engine/Source/Runtime/MediaAssets/Private/Assets/MediaSoundComponent.cpp b/Engine/Source/Runtime/MediaAssets/Private/Assets/MediaSoundComponent.cpp
--- a/Engine/Source/Runtime/MediaAssets/Private/Assets/MediaSoundComponent.cpp
+++ b/Engine/Source/Runtime/MediaAssets/Private/Assets/MediaSoundComponent.cpp
@@ -609,6 +609,7 @@ int32 FMediaSoundGenerator::OnGenerateAudio(float* OutAudio, int32 NumSamples)
 		Resampler.Flush();
 
 		LastPlaySampleTime = FTimespan::MinValue();
+		FMemory::Memzero(OutAudio, NumSamples * sizeof(float));
 	}
  	return NumSamples;
 
