Adding checks to Mediaframework's texture conversion code to safely ignore samples that present no converter, texture or CPU side buffer despite otherwise being valid
#rnx
#tests Local tests with various medi aplayer plugins

diff --git a/Engine/Source/Runtime/MediaAssets/Private/Misc/MediaTextureResource.cpp b/Engine/Source/Runtime/MediaAssets/Private/Misc/MediaTextureResource.cpp
--- a/Engine/Source/Runtime/MediaAssets/Private/Misc/MediaTextureResource.cpp
+++ b/Engine/Source/Runtime/MediaAssets/Private/Misc/MediaTextureResource.cpp
@@ -518,114 +518,118 @@ void FMediaTextureResource::Render(const FRenderParams& Params)
 			else
 			{
 				// No custom conversion, we need default processing...
-				ConvertOrCopyNeeded = true;
-			}
 
-			if (ConvertOrCopyNeeded)
-			{
-				if (RequiresConversion(Sample, NumMips))
+				// Sanity check: do we have any data?
+				if (Sample->GetTexture() == nullptr && Sample->GetBuffer() == nullptr)
 				{
-					//
-					// Sample needs to be converted by built in converter code
-					//
-					ConvertSample(Sample, Params.ClearColor, NumMips);
+					UseSample = false;
 				}
 				else
 				{
-					//
-					// Sample can be used directly or is a simple copy
-					//
-					CopySample(Sample, Params.ClearColor, NumMips, Params.CurrentGuid);
+					ConvertOrCopyNeeded = true;
 				}
 			}
 
-			if (IMediaTextureSampleColorConverter* Converter = Sample->GetMediaTextureSampleColorConverter())
+			if (UseSample)
 			{
-				FTexture2DRHIRef TextureRef = IntermediateTarget ? IntermediateTarget : RenderTargetTextureRHI;
+				if (ConvertOrCopyNeeded)
+				{
+					if (RequiresConversion(Sample, NumMips))
+					{
+						//
+						// Sample needs to be converted by built in converter code
+						//
+						ConvertSample(Sample, Params.ClearColor, NumMips);
+					}
+					else
+					{
+						//
+						// Sample can be used directly or is a simple copy
+						//
+						CopySample(Sample, Params.ClearColor, NumMips, Params.CurrentGuid);
+					}
+				}
 
-				if (bRecreateOutputTarget && TextureRef)
+				if (IMediaTextureSampleColorConverter* Converter = Sample->GetMediaTextureSampleColorConverter())
 				{
-					const static FLazyName ClassName(TEXT("FMediaTextureResource"));
+					FTexture2DRHIRef TextureRef = IntermediateTarget ? IntermediateTarget : RenderTargetTextureRHI;
 
-					const FRHITextureDesc& IntermediateTextureDesc = TextureRef->GetDesc();
-			
-					const FRHITextureCreateDesc Desc =
-						FRHITextureCreateDesc::Create2D(TEXT("MediaTextureColorConversionOutput"),
-							IntermediateTextureDesc.GetSize().X,
-							IntermediateTextureDesc.GetSize().Y,
-							IntermediateTextureDesc.Format)
-						.SetNumMips(IntermediateTextureDesc.NumMips)
-						.SetClearValue(IntermediateTextureDesc.ClearValue)
-						.SetClassName(ClassName)
-						.SetOwnerName(GetOwnerName())
-						.SetFlags(IntermediateTextureDesc.Flags);
-			
-					OutputTarget = RHICreateTexture(Desc);
-					OutputTarget->SetName(TEXT("MediaTextureResourceOutput"));
-					OutputTarget->SetOwnerName(GetOwnerName());
-			
-					bRecreateOutputTarget = false;
-				}
-		
-				Converter->ApplyColorConversion(TextureRef, OutputTarget);
+					if (bRecreateOutputTarget && TextureRef)
+					{
+						const static FLazyName ClassName(TEXT("FMediaTextureResource"));
+
+						const FRHITextureDesc& IntermediateTextureDesc = TextureRef->GetDesc();
+
+						const FRHITextureCreateDesc Desc =
+							FRHITextureCreateDesc::Create2D(TEXT("MediaTextureColorConversionOutput"),
+								IntermediateTextureDesc.GetSize().X,
+								IntermediateTextureDesc.GetSize().Y,
+								IntermediateTextureDesc.Format)
+							.SetNumMips(IntermediateTextureDesc.NumMips)
+							.SetClearValue(IntermediateTextureDesc.ClearValue)
+							.SetClassName(ClassName)
+							.SetOwnerName(GetOwnerName())
+							.SetFlags(IntermediateTextureDesc.Flags);
+
+						OutputTarget = RHICreateTexture(Desc);
+						OutputTarget->SetName(TEXT("MediaTextureResourceOutput"));
+						OutputTarget->SetOwnerName(GetOwnerName());
+
+						bRecreateOutputTarget = false;
+					}
 
-				if (RenderTargetTextureRHI != OutputTarget)
+					Converter->ApplyColorConversion(TextureRef, OutputTarget);
+
+					if (RenderTargetTextureRHI != OutputTarget)
+					{
+						UpdateTextureReference(OutputTarget);
+					}
+				}
+				else if (IntermediateTarget)
 				{
-					UpdateTextureReference(OutputTarget);
+					OutputTarget = IntermediateTarget;
 				}
-			}
-			else if (IntermediateTarget)
-			{
-				OutputTarget = IntermediateTarget;
-			}
 
-			Rotation = Sample->GetScaleRotation();
-			Offset = Sample->GetOffset();
+				Rotation = Sample->GetScaleRotation();
+				Offset = Sample->GetOffset();
 
-			if (CurrentSample)
-			{
-				// If we had a current sample (directly used as output), we can now schedule its retirement
-				LocalPriorSamples->Retire(CurrentSample);
-				CurrentSample = nullptr;
-			}
-
-			// Do we use a local copy as our output?
-			if (OutputTarget == RenderTargetTextureRHI)
-			{
-				// Yes, we can schedule the actual sample for retirement right away
-				LocalPriorSamples->Retire(Sample);
-			}
-			else
-			{
-				// No, we need to hold on to the sample
-				CurrentSample = Sample;
-			}
+				if (CurrentSample)
+				{
+					// If we had a current sample (directly used as output), we can now schedule its retirement
+					LocalPriorSamples->Retire(CurrentSample);
+					CurrentSample = nullptr;
+				}
 
-			// Generate mips as needed
-			if (CurrentNumMips > 1 && !Cleared && Sample->GetNumMips() == 1)
-			{
-				check(OutputTarget);
+				// Do we use a local copy as our output?
+				if (OutputTarget == RenderTargetTextureRHI)
+				{
+					// Yes, we can schedule the actual sample for retirement right away
+					LocalPriorSamples->Retire(Sample);
+				}
+				else
+				{
+					// No, we need to hold on to the sample
+					CurrentSample = Sample;
+				}
 
-				const EGenerateMipsPass GenerateMipsPass =
-					MediaTextureResourceHelpers::SupportsComputeMipGen(OutputTarget->GetFormat()) ? EGenerateMipsPass::Compute : EGenerateMipsPass::Raster;
+				// Generate mips as needed
+				if (CurrentNumMips > 1 && !Cleared && Sample->GetNumMips() == 1)
+				{
+					check(OutputTarget);
 
-				CacheRenderTarget(OutputTarget, TEXT("MipGeneration"), MipGenerationCache);
+					const EGenerateMipsPass GenerateMipsPass =
+						MediaTextureResourceHelpers::SupportsComputeMipGen(OutputTarget->GetFormat()) ? EGenerateMipsPass::Compute : EGenerateMipsPass::Raster;
 
-				FRDGBuilder GraphBuilder(FRHICommandListExecutor::GetImmediateCommandList());
-				FRDGTextureRef MipOutputTexture = GraphBuilder.RegisterExternalTexture(MipGenerationCache);
-				FGenerateMips::Execute(GraphBuilder, GetFeatureLevel(), MipOutputTexture, FGenerateMipsParams{ SF_Bilinear, AM_Clamp, AM_Clamp, AM_Clamp }, GenerateMipsPass);
-				GraphBuilder.Execute();
-			}
+					CacheRenderTarget(OutputTarget, TEXT("MipGeneration"), MipGenerationCache);
 
-			SET_FLOAT_STAT(STAT_MediaUtils_TextureSampleTime, Sample->GetTime().Time.GetTotalMilliseconds());
-		}
-		else
-		{
-			//
-			// Last sample is still valid
-			//
+					FRDGBuilder GraphBuilder(FRHICommandListExecutor::GetImmediateCommandList());
+					FRDGTextureRef MipOutputTexture = GraphBuilder.RegisterExternalTexture(MipGenerationCache);
+					FGenerateMips::Execute(GraphBuilder, GetFeatureLevel(), MipOutputTexture, FGenerateMipsParams{ SF_Bilinear, AM_Clamp, AM_Clamp, AM_Clamp }, GenerateMipsPass);
+					GraphBuilder.Execute();
+				}
 
-			// nothing to do for now
+				SET_FLOAT_STAT(STAT_MediaUtils_TextureSampleTime, Sample->GetTime().Time.GetTotalMilliseconds());
+			}
 		}
 	}
 	else if (Params.CanClear)
@@ -997,6 +1001,7 @@ void FMediaTextureResource::ConvertSample(const TSharedPtr<IMediaTextureSample,
 
 			// copy sample data to input render target
 			const uint8* Data = (const uint8*)Sample->GetBuffer();
+			check(Data);
 			for (uint8 MipLevel = 0; MipLevel < SampleNumMips; ++MipLevel)
 			{
 				uint32 Stride = Sample->GetStride() >> MipLevel;
