Media: Deprecated media texture output format property.

diff --git a/Engine/Source/Runtime/MediaAssets/Private/Assets/MediaTexture.cpp b/Engine/Source/Runtime/MediaAssets/Private/Assets/MediaTexture.cpp
--- a/Engine/Source/Runtime/MediaAssets/Private/Assets/MediaTexture.cpp
+++ b/Engine/Source/Runtime/MediaAssets/Private/Assets/MediaTexture.cpp
@@ -79,7 +79,6 @@ UMediaTexture::UMediaTexture(const FObjectInitializer& ObjectInitializer)
 	, EnableGenMips(false)
 	, NumMips(1)
 	, NewStyleOutput(false)
-	, OutputFormat(MTOF_Default)
 	, CurrentAspectRatio(0.0f)
 	, CurrentOrientation(MTORI_Original)
 	, DefaultGuid(FGuid::NewGuid())
diff --git a/Engine/Source/Runtime/MediaAssets/Private/Misc/MediaTextureResource.cpp b/Engine/Source/Runtime/MediaAssets/Private/Misc/MediaTextureResource.cpp
--- a/Engine/Source/Runtime/MediaAssets/Private/Misc/MediaTextureResource.cpp
+++ b/Engine/Source/Runtime/MediaAssets/Private/Misc/MediaTextureResource.cpp
@@ -838,9 +838,6 @@ bool FMediaTextureResource::RequiresConversion(const FTextureRHIRef& SampleTextu
 		// New Style
 		// 
 
-		// For now we only allow this, single SRGB-style output format
-		check(Owner.OutputFormat == MTOF_SRGB_LINOUT || Owner.OutputFormat == MTOF_Default);
-
 		// If we have no mips in the sample, but want to have some in the output, we use the conversion pass
 		// to setup level 0 and have a suitable output texture
 		if (SampleTexture->GetNumMips() == 1 && InNumMips != 1)
@@ -875,9 +872,6 @@ bool FMediaTextureResource::RequiresConversion(const TSharedPtr<IMediaTextureSam
 		// New Style
 		// 
 
-		// For now we only allow this, single SRGB-style output format
-		check(Owner.OutputFormat == MTOF_SRGB_LINOUT || Owner.OutputFormat == MTOF_Default);
-
 		FRHITexture *Texture = Sample->GetTexture();
 
 		// If we have no mips in the sample, but want to have some in the output, we use the conversion pass
diff --git a/Engine/Source/Runtime/MediaAssets/Public/MediaTexture.h b/Engine/Source/Runtime/MediaAssets/Public/MediaTexture.h
--- a/Engine/Source/Runtime/MediaAssets/Public/MediaTexture.h
+++ b/Engine/Source/Runtime/MediaAssets/Public/MediaTexture.h
@@ -22,7 +22,7 @@ class IMediaTextureSample;
 class UMediaPlayer;
 
 UENUM()
-enum MediaTextureOutputFormat : int
+enum UE_DEPRECATED(5.4, "No longer used") MediaTextureOutputFormat : int
 {
 	MTOF_Default					UMETA(DisplayName = "Default (sRGB)"),
 	MTOF_SRGB_LINOUT				UMETA(DisplayName = "sRGB (linear output)"),		// sRGB data, using sRGB texture formats; hence read as linear RGB
@@ -89,8 +89,13 @@ public:
 	UPROPERTY(EditAnywhere, BlueprintReadWrite, Category = "MediaTexture", meta = (DisplayName = "Enable new style output"))
 	bool NewStyleOutput;
 
-	UPROPERTY(EditAnywhere, BlueprintReadWrite, Category = "MediaTexture", meta = (DisplayName = "Output format (new style)"))
-	TEnumAsByte<enum MediaTextureOutputFormat> OutputFormat;
+PRAGMA_DISABLE_DEPRECATION_WARNINGS
+#if WITH_EDITORONLY_DATA
+	/** DEPRECATED 5.4 */
+	UPROPERTY(meta = (DeprecatedProperty, DeprecationMessage = "Output format is now deprecated."))
+	TEnumAsByte<enum MediaTextureOutputFormat> OutputFormat_DEPRECATED;
+#endif
+PRAGMA_ENABLE_DEPRECATION_WARNINGS
 
 	/** Current aspect ratio */
 	UPROPERTY(Transient, TextExportTransient, SkipSerialization, BlueprintReadOnly, Category = "MediaTexture", meta = (DisplayName = "Current frame's aspect ratio"))
