MediaPlayer: Not setting the media player's internal loop state when looping is enabled while playing a playlist. Looping a playlist is to loop the playlist, not an individual video that is currently playing when looping is enabled as that might result in the same video entry being looped endlessly without advancing to the next playlist entry.

#rnx

diff --git a/Engine/Source/Runtime/MediaAssets/Private/Assets/MediaPlayer.cpp b/Engine/Source/Runtime/MediaAssets/Private/Assets/MediaPlayer.cpp
--- a/Engine/Source/Runtime/MediaAssets/Private/Assets/MediaPlayer.cpp
+++ b/Engine/Source/Runtime/MediaAssets/Private/Assets/MediaPlayer.cpp
@@ -826,7 +826,7 @@ bool UMediaPlayer::SetLooping(bool Looping)
 
 	Loop = Looping;
 
-	return PlayerFacade->SetLooping(Looping);
+	return PlayerFacade->SetLooping(Loop && (!Playlist || Playlist->Num() == 1));
 }
 
 
