Media: Improved deprecation message on recently deprecated unused output format property.

#rnx

diff --git a/Engine/Source/Runtime/MediaAssets/Public/MediaTexture.h b/Engine/Source/Runtime/MediaAssets/Public/MediaTexture.h
--- a/Engine/Source/Runtime/MediaAssets/Public/MediaTexture.h
+++ b/Engine/Source/Runtime/MediaAssets/Public/MediaTexture.h
@@ -22,7 +22,7 @@ class IMediaTextureSample;
 class UMediaPlayer;
 
 UENUM()
-enum UE_DEPRECATED(5.4, "No longer used") MediaTextureOutputFormat : int
+enum UE_DEPRECATED(5.4, "This enum was unused (not connected to active logic) and is now deprecated.") MediaTextureOutputFormat : int
 {
 	MTOF_Default					UMETA(DisplayName = "Default (sRGB)"),
 	MTOF_SRGB_LINOUT				UMETA(DisplayName = "sRGB (linear output)"),		// sRGB data, using sRGB texture formats; hence read as linear RGB
@@ -92,7 +92,7 @@ public:
 PRAGMA_DISABLE_DEPRECATION_WARNINGS
 #if WITH_EDITORONLY_DATA
 	/** DEPRECATED 5.4 */
-	UPROPERTY(meta = (DeprecatedProperty, DeprecationMessage = "Output format is now deprecated."))
+	UPROPERTY(meta = (DeprecatedProperty, DeprecationMessage = "Output format was unused (not connected to active logic) and is now deprecated. References to it can be safely deleted."))
 	TEnumAsByte<enum MediaTextureOutputFormat> OutputFormat_DEPRECATED;
 #endif
 PRAGMA_ENABLE_DEPRECATION_WARNINGS
