Media: Adding provisions to allow for a playback range of a media stream

#rnx

diff --git a/Engine/Plugins/Media/ElectraPlayer/Source/ElectraPlayerPlugin/Private/ElectraPlayerPlugin.cpp b/Engine/Plugins/Media/ElectraPlayer/Source/ElectraPlayerPlugin/Private/ElectraPlayerPlugin.cpp
--- a/Engine/Plugins/Media/ElectraPlayer/Source/ElectraPlayerPlugin/Private/ElectraPlayerPlugin.cpp
+++ b/Engine/Plugins/Media/ElectraPlayer/Source/ElectraPlayerPlugin/Private/ElectraPlayerPlugin.cpp
@@ -61,6 +61,9 @@ FElectraPlayerPlugin::FElectraPlayerPlugin()
 	static_assert((int32)EMediaRateThinning::Unthinned == (int32)IElectraPlayerInterface::EPlayRateType::Unthinned, "check alignment of both enums");
 	static_assert((int32)EMediaRateThinning::Thinned == (int32)IElectraPlayerInterface::EPlayRateType::Thinned, "check alignment of both enums");
 
+	static_assert((int32)EMediaTimeRangeType::Absolute == (int32)IElectraPlayerInterface::ETimeRangeType::Absolute, "check alignment of both enums");
+	static_assert((int32)EMediaTimeRangeType::Current == (int32)IElectraPlayerInterface::ETimeRangeType::Current, "check alignment of both enums");
+
 	static_assert(IMediaPlayerLifecycleManagerDelegate::ResourceFlags_Decoder == IElectraPlayerInterface::ResourceFlags_Decoder, "check alignment of both enums");
 	static_assert(IMediaPlayerLifecycleManagerDelegate::ResourceFlags_OutputBuffers == IElectraPlayerInterface::ResourceFlags_OutputBuffers, "check alignment of both enums");
 	static_assert(IMediaPlayerLifecycleManagerDelegate::ResourceFlags_Any == IElectraPlayerInterface::ResourceFlags_Any, "check alignment of both enums");
@@ -901,6 +904,10 @@ bool FElectraPlayerPlugin::CanControl(EMediaControl Control) const
 	{
 		return CurrentState == EMediaState::Playing || CurrentState == EMediaState::Paused || CurrentState == EMediaState::Stopped;
 	}
+	else if (Control == EMediaControl::PlaybackRange)
+	{
+		return true;
+	}
 	return false;
 }
 
@@ -982,6 +989,20 @@ bool FElectraPlayerPlugin::Seek(const FTimespan& Time)
 	return Player->Seek(Time);
 }
 
+TRange<FTimespan> FElectraPlayerPlugin::GetPlaybackTimeRange(EMediaTimeRangeType InRangeToGet) const
+{
+	return Player->GetPlaybackRange(static_cast<IElectraPlayerInterface::ETimeRangeType>(InRangeToGet));
+}
+
+bool FElectraPlayerPlugin::SetPlaybackTimeRange(const TRange<FTimespan>& InTimeRange)
+{
+	IElectraPlayerInterface::FPlaybackRange Range;
+	Range.Start = InTimeRange.GetLowerBoundValue();
+	Range.End = InTimeRange.GetUpperBoundValue();
+	Player->SetPlaybackRange(Range);
+	return true;
+}
+
 
 bool FElectraPlayerPlugin::GetAudioTrackFormat(int32 TrackIndex, int32 FormatIndex, FMediaAudioTrackFormat& OutFormat) const
 {
diff --git a/Engine/Plugins/Media/ElectraPlayer/Source/ElectraPlayerPlugin/Private/ElectraPlayerPlugin.h b/Engine/Plugins/Media/ElectraPlayer/Source/ElectraPlayerPlugin/Private/ElectraPlayerPlugin.h
--- a/Engine/Plugins/Media/ElectraPlayer/Source/ElectraPlayerPlugin/Private/ElectraPlayerPlugin.h
+++ b/Engine/Plugins/Media/ElectraPlayer/Source/ElectraPlayerPlugin/Private/ElectraPlayerPlugin.h
@@ -127,6 +127,8 @@ private:
 	bool SetRate(float Rate) override;
 	bool Seek(const FTimespan& Time) override;
 
+	TRange<FTimespan> GetPlaybackTimeRange(EMediaTimeRangeType InRangeToGet) const override;
+	bool SetPlaybackTimeRange(const TRange<FTimespan>& InTimeRange) override;
 
 	// From IMediaTracks
 	bool GetAudioTrackFormat(int32 TrackIndex, int32 FormatIndex, FMediaAudioTrackFormat& OutFormat) const override;
diff --git a/Engine/Plugins/Media/ElectraPlayer/Source/ElectraPlayerRuntime/Private/ElectraPlayer.cpp b/Engine/Plugins/Media/ElectraPlayer/Source/ElectraPlayerRuntime/Private/ElectraPlayer.cpp
--- a/Engine/Plugins/Media/ElectraPlayer/Source/ElectraPlayerRuntime/Private/ElectraPlayer.cpp
+++ b/Engine/Plugins/Media/ElectraPlayer/Source/ElectraPlayerRuntime/Private/ElectraPlayer.cpp
@@ -336,9 +336,9 @@ bool FElectraPlayer::OpenInternal(const FString& Url, const FParamDict& InPlayer
 			NewPlayer->AdaptivePlayer->SetBitrateCeiling(PlaystartOptions.MaxBandwidthForStreaming.GetValue());
 		}
 
-        // Set the player member variable to the new player so we can use our internal configuration methods on the new player. 
+        // Set the player member variable to the new player so we can use our internal configuration methods on the new player.
         CurrentPlayer = MoveTemp(NewPlayer);
-		
+
 		// Apply options that may have been set prior to calling Open().
 		// Set these only if they have defined values as to not override what might have been set in the PlayerOptions.
 		if (bFrameAccurateSeeking.IsSet())
@@ -1266,6 +1266,12 @@ void FElectraPlayer::SetPlaybackRange(const FPlaybackRange& InPlaybackRange)
 	TSharedPtr<FInternalPlayerImpl, ESPMode::ThreadSafe> LockedPlayer = CurrentPlayer;
 	if (LockedPlayer.IsValid() && LockedPlayer->AdaptivePlayer.IsValid())
 	{
+		// Ranges cannot be set on Live streams.
+		Electra::FTimeValue playDuration = CurrentPlayer->AdaptivePlayer->GetDuration();
+		if (playDuration.IsValid() && playDuration.IsInfinity())
+		{
+			return;
+		}
 		Electra::IAdaptiveStreamingPlayer::FPlaybackRange Range;
 		if (CurrentPlaybackRange.Start.IsSet())
 		{
@@ -1306,6 +1312,55 @@ void FElectraPlayer::GetPlaybackRange(FPlaybackRange& OutPlaybackRange) const
 	}
 }
 
+TRange<FTimespan> FElectraPlayer::GetPlaybackRange(ETimeRangeType InRangeToGet) const
+{
+	TRange<FTimespan> Range(FTimespan(0), FTimespan(0));
+	TSharedPtr<FInternalPlayerImpl, ESPMode::ThreadSafe> LockedPlayer = CurrentPlayer;
+	if (LockedPlayer.IsValid() && LockedPlayer->AdaptivePlayer.IsValid())
+	{
+		switch(InRangeToGet)
+		{
+			case IElectraPlayerInterface::ETimeRangeType::Absolute:
+			{
+				Electra::FTimeRange Timeline;
+				LockedPlayer->AdaptivePlayer->GetTimelineRange(Timeline);
+				if (Timeline.IsValid())
+				{
+					Range.SetLowerBound(Timeline.Start.GetAsTimespan());
+					Range.SetUpperBound(Timeline.End.GetAsTimespan());
+				}
+				else
+				{
+					Electra::FTimeValue playDuration = CurrentPlayer->AdaptivePlayer->GetDuration();
+					if (playDuration.IsValid())
+					{
+						Range.SetLowerBound(FTimespan(0));
+						Range.SetUpperBound(playDuration.IsInfinity() ? FTimespan::MaxValue() : playDuration.GetAsTimespan());
+					}
+				}
+				break;
+			}
+			case IElectraPlayerInterface::ETimeRangeType::Current:
+			{
+				Electra::IAdaptiveStreamingPlayer::FPlaybackRange Current;
+				LockedPlayer->AdaptivePlayer->GetPlaybackRange(Current);
+				if (Current.Start.IsSet() && Current.End.IsSet())
+				{
+					Range.SetLowerBound(Current.Start.GetValue().GetAsTimespan());
+					Range.SetUpperBound(Current.End.GetValue().GetAsTimespan());
+				}
+				else
+				{
+					return GetPlaybackRange(IElectraPlayerInterface::ETimeRangeType::Absolute);
+				}
+				break;
+			}
+		}
+	}
+	return Range;
+}
+
+
 TSharedPtr<TMap<FString, TArray<TSharedPtr<Electra::IMediaStreamMetadata::IItem, ESPMode::ThreadSafe>>>, ESPMode::ThreadSafe> FElectraPlayer::GetMediaMetadata() const
 {
 	return CurrentStreamMetadata;
@@ -2775,7 +2830,7 @@ void FElectraPlayer::HandlePlayerMediaMetadataChanged(const TSharedPtrTS<Electra
 	if (InMetadata.IsValid())
 	{
 		TSharedPtr<TMap<FString, TArray<TSharedPtr<Electra::IMediaStreamMetadata::IItem, ESPMode::ThreadSafe>>>, ESPMode::ThreadSafe> NewMeta = InMetadata->GetMediaStreamMetadata();
-		CurrentStreamMetadata = MoveTemp(NewMeta);		
+		CurrentStreamMetadata = MoveTemp(NewMeta);
 		DeferredEvents.Enqueue(IElectraPlayerAdapterDelegate::EPlayerEvent::MetadataChanged);
 
 		TSharedPtr<IElectraPlayerAdapterDelegate, ESPMode::ThreadSafe> PinnedAdapterDelegate = AdapterDelegate.Pin();
@@ -3020,7 +3075,7 @@ void FElectraPlayer::LogStatistics()
 			*Statistics.LastState,
 			*SanitizeMessage(Statistics.LastError)
 		);
-		
+
 		if (Statistics.LastError.Len())
 		{
 			FString MessageHistory;
diff --git a/Engine/Plugins/Media/ElectraPlayer/Source/ElectraPlayerRuntime/Private/ElectraPlayer.h b/Engine/Plugins/Media/ElectraPlayer/Source/ElectraPlayerRuntime/Private/ElectraPlayer.h
--- a/Engine/Plugins/Media/ElectraPlayer/Source/ElectraPlayerRuntime/Private/ElectraPlayer.h
+++ b/Engine/Plugins/Media/ElectraPlayer/Source/ElectraPlayerRuntime/Private/ElectraPlayer.h
@@ -110,6 +110,7 @@ public:
 
 	void SetPlaybackRange(const FPlaybackRange& InPlaybackRange) override;
 	void GetPlaybackRange(FPlaybackRange& OutPlaybackRange) const override;
+	TRange<FTimespan> GetPlaybackRange(ETimeRangeType InRangeToGet) const override;
 
 	TRangeSet<float> GetSupportedRates(EPlayRateType InPlayRateType) const override;
 	float GetRate() const override;
diff --git a/Engine/Plugins/Media/ElectraPlayer/Source/ElectraPlayerRuntime/Public/IElectraPlayerInterface.h b/Engine/Plugins/Media/ElectraPlayer/Source/ElectraPlayerRuntime/Public/IElectraPlayerInterface.h
--- a/Engine/Plugins/Media/ElectraPlayer/Source/ElectraPlayerRuntime/Public/IElectraPlayerInterface.h
+++ b/Engine/Plugins/Media/ElectraPlayer/Source/ElectraPlayerRuntime/Public/IElectraPlayerInterface.h
@@ -181,7 +181,7 @@ public:
 		TOptional<int32>			MaxVerticalStreamResolution;
 		TOptional<int32>			MaxBandwidthForStreaming;
 		bool						bDoNotPreload = false;
-		
+
 		TSharedPtr<IElectraPlayerExternalDataReader, ESPMode::ThreadSafe> ExternalDataReader;
 		TSharedPtr<IElectraPlayerDataCache, ESPMode::ThreadSafe> ExternalDataCache;
 	};
@@ -237,6 +237,17 @@ public:
 	virtual void SetPlaybackRange(const FPlaybackRange& InPlaybackRange) = 0;
 	virtual void GetPlaybackRange(FPlaybackRange& OutPlaybackRange) const = 0;
 
+	enum class ETimeRangeType
+	{
+		/** Total absolute time range as defined by the media. */
+		Absolute,
+
+		/** Current time range of the media, set by media internal means or through API calls. */
+		Current
+	};
+	virtual TRange<FTimespan> GetPlaybackRange(ETimeRangeType InRangeToGet) const = 0;
+
+
 	enum class EPlayRateType
 	{
 		// Playback rate with frames being dropped.
@@ -260,7 +271,7 @@ public:
 	virtual bool Seek(const FTimespan& Time) = 0;
 	virtual bool Seek(const FTimespan& Time, const FSeekParam& Param) = 0;
 	virtual void SetFrameAccurateSeekMode(bool bEnableFrameAccuracy) = 0;
-	
+
 	virtual void ModifyOptions(const Electra::FParamDict& InOptionsToSetOrChange, const Electra::FParamDict& InOptionsToClear) = 0;
 
 	struct FAudioTrackFormat
diff --git a/Engine/Source/Runtime/Media/Public/IMediaControls.h b/Engine/Source/Runtime/Media/Public/IMediaControls.h
--- a/Engine/Source/Runtime/Media/Public/IMediaControls.h
+++ b/Engine/Source/Runtime/Media/Public/IMediaControls.h
@@ -26,7 +26,10 @@ enum class EMediaControl
 	Scrub,
 
 	/** Seek to playback position. */
-	Seek
+	Seek,
+
+	/** Playback of a range. */
+	PlaybackRange
 };
 
 
@@ -106,6 +109,19 @@ enum class EMediaStatus
 ENUM_CLASS_FLAGS(EMediaStatus);
 
 
+/**
+ * Different types of media timeline ranges.
+ */
+enum class EMediaTimeRangeType
+{
+	/** Total absolute time range as defined by the media. */
+	Absolute,
+
+	/** Current time range of the media, set by media internal means or through API calls. */
+	Current
+};
+
+
 /**
  * Interface for controlling media playback.
  *
@@ -211,9 +227,9 @@ public:
 
 	/**
 	 * Hint for player indicating that blocked playback mode will be used / not used
-	 * 
+	 *
 	 * @param bFacadeWillUseBlockingPlayback True if blocked playback will be used, false otherwiese
-	 * 
+	 *
 	 * @note Implement as needed by the player.
 	 */
 	virtual void SetBlockingPlaybackHint(bool bFacadeWillUseBlockingPlayback)
@@ -282,6 +298,16 @@ public:
 		return Seek(SeekTime);
 	}
 
+	virtual TRange<FTimespan> GetPlaybackTimeRange(EMediaTimeRangeType InRangeToGet) const
+	{
+		return TRange<FTimespan>(FTimespan(0), GetDuration());
+	}
+
+	virtual bool SetPlaybackTimeRange(const TRange<FTimespan>& InTimeRange)
+	{
+		return false;
+	}
+
 public:
 
 	/** Virtual destructor. */
diff --git a/Engine/Source/Runtime/MediaAssets/Private/Assets/MediaPlayer.cpp b/Engine/Source/Runtime/MediaAssets/Private/Assets/MediaPlayer.cpp
--- a/Engine/Source/Runtime/MediaAssets/Private/Assets/MediaPlayer.cpp
+++ b/Engine/Source/Runtime/MediaAssets/Private/Assets/MediaPlayer.cpp
@@ -905,6 +905,38 @@ bool UMediaPlayer::SupportsSeeking() const
 	return PlayerFacade->CanSeek();
 }
 
+bool UMediaPlayer::SupportsPlaybackTimeRange() const
+{
+	return PlayerFacade->SupportsPlaybackTimeRange();
+}
+
+TRange<FTimespan> UMediaPlayer::GetPlaybackTimeRange(EMediaTimeRangeType InRangeToGet)
+{
+	return PlayerFacade->GetPlaybackTimeRange(InRangeToGet);
+}
+
+bool UMediaPlayer::SetPlaybackTimeRange(const TRange<FTimespan>& InTimeRange)
+{
+	return PlayerFacade->SetPlaybackTimeRange(InTimeRange);
+}
+
+FFloatInterval UMediaPlayer::GetPlaybackTimeRange(EMediaTimeRangeBPType InRangeToGet)
+{
+	TRange<FTimespan> r = GetPlaybackTimeRange(static_cast<EMediaTimeRangeType>(InRangeToGet));
+	FDoubleInterval i;
+	i.Min = !r.IsEmpty() ? r.GetLowerBoundValue().GetTotalSeconds() : -1.0;
+	i.Max = !r.IsEmpty() ? r.GetUpperBoundValue().GetTotalSeconds() : -1.0;
+	FFloatInterval Imprecise;
+	Imprecise.Min = (float) i.Min;
+	Imprecise.Max = (float) i.Max;
+	return Imprecise;
+}
+
+bool UMediaPlayer::SetPlaybackTimeRange(FFloatInterval InTimeRange)
+{
+	return SetPlaybackTimeRange(TRange<FTimespan>(FTimespan::FromSeconds(InTimeRange.Min), FTimespan::FromSeconds(InTimeRange.Max)));
+}
+
 
 #if WITH_EDITOR
 
diff --git a/Engine/Source/Runtime/MediaAssets/Public/MediaPlayer.h b/Engine/Source/Runtime/MediaAssets/Public/MediaPlayer.h
--- a/Engine/Source/Runtime/MediaAssets/Public/MediaPlayer.h
+++ b/Engine/Source/Runtime/MediaAssets/Public/MediaPlayer.h
@@ -7,6 +7,7 @@
 #include "Delegates/Delegate.h"
 #include "IMediaOptions.h"
 #include "IMediaMetadataItem.h"
+#include "IMediaControls.h"
 #include "Math/Quat.h"
 #include "Math/Rotator.h"
 #include "Templates/SharedPointer.h"
@@ -74,6 +75,21 @@ enum class EMediaPlayerTrack : uint8
 };
 
 
+/**
+ * Blueprint usable enum of EMediaTimeRangeType
+ */
+UENUM(BlueprintType)
+enum class EMediaTimeRangeBPType : uint8
+{
+	/** Total absolute time range as defined by the media. */
+	Absolute = (uint8)EMediaTimeRangeType::Absolute,
+
+	/** Current time range of the media, set by media internal means or through API calls. */
+	Current = (uint8)EMediaTimeRangeType::Current
+
+};
+
+
 UCLASS(BlueprintType, hidecategories = (Object), MinimalAPI)
 class UMediaTimeStampInfo
 	: public UObject
@@ -98,13 +114,13 @@ struct FMediaMetadataItemBPT
 public:
 	UPROPERTY(BlueprintReadOnly, Category = "Media|Metadata")
 	FString LanguageCode;
-	
+
 	UPROPERTY(BlueprintReadOnly, Category = "Media|Metadata")
 	FString MimeType;
-	
+
 	UPROPERTY(BlueprintReadOnly, Category = "Media|Metadata")
 	FString StringData;
-	
+
 	UPROPERTY(BlueprintReadOnly, Category = "Media|Metadata")
 	TArray<uint8> BinaryData;
 };
@@ -1037,6 +1053,72 @@ public:
 	FOnMediaPlayerMediaEvent OnMetadataChanged;
 public:
 
+	/**
+	 * Check whether the player supports playing back of range within the media.
+	 *
+	 * @return true if playing back a range is supported, false otherwise.
+	 * @see GetPlaybackTimeRange, SetPlaybackTimeRange
+	 */
+	UFUNCTION(BlueprintCallable, Category="Media|MediaPlayer")
+	MEDIAASSETS_API bool SupportsPlaybackTimeRange() const;
+
+	/**
+	 * Returns the current playback range of the media.
+	 * If playing back a range is not supported, the range returned will be equal
+	 * to [ 0, GetDuration() ).
+	 * The media may have an implicit default range provided by the container format
+	 * or other means without having called SetPlaybackTimeRange().
+	 * The media may have internal time values not starting at 0, which are
+	 * conveyed by the range.
+	 * Since the range may be only a portion of the media, the duration of the
+	 * returned range may be less than the media overall duration returned by
+	 * GetDuration().
+	 * For live video streams the range may change dynamically as new content
+	 * becomes available and old content falls off the timeline.
+	 *
+	 * @param InRangeToGet The type of range to get.
+	 *                     `Absolute` returns the media's smallest and largest timeline values.
+	 *                       Unless continuously changing in a Live stream this is usually the
+	 *                       same as [ 0, GetDuration() ]. The base time does not have to be
+	 *                       zero though.
+	 *                     `Current` returns the currently set range, which is a subset of the
+	 *                       absolute range.
+	 * @return The playback range as queried for.
+	 * @see SupportsPlaybackTimeRange, SetPlaybackTimeRange, GetDuration
+	 */
+	MEDIAASSETS_API TRange<FTimespan> GetPlaybackTimeRange(EMediaTimeRangeType InRangeToGet);
+
+	/**
+	 * Blueprint accessible version of GetPlaybackTimeRange.
+	 * This returns the range truncated into a blueprint usable float interval and should not
+	 * be used for live streams as 32 bit floats can not store wallclock times with enough precision.
+	 */
+	UFUNCTION(BlueprintCallable, Category="Media|MediaPlayer")
+	MEDIAASSETS_API FFloatInterval GetPlaybackTimeRange(EMediaTimeRangeBPType InRangeToGet);
+
+	/**
+	 * Sets a new media playback range.
+	 * Has an effect only if SupportsPlaybackTimeRange() returns true and the media supports it.
+	 * A live stream cannot be constrained to a range.
+	 * The range will be clamped if necessary to be within the media's absolute time range.
+	 * Changing the time range may trigger an implicit Seek() depending on where the current
+	 * playback position is located with regard to the new range.
+	 * Unless prevented by the media a playback range can be cleared by passing an empty range.
+	 *
+	 * @param InTimeRange The new playback range to set.
+	 * @return true if successful, false otherwise.
+	 * @see SupportsPlaybackTimeRange, GetPlaybackTimeRange
+	 */
+	MEDIAASSETS_API bool SetPlaybackTimeRange(const TRange<FTimespan>& InTimeRange);
+
+	/**
+	 * Blueprint accessible version of SetPlaybackTimeRange().
+	 * The range is set through a blueprint usable float interval which may not have enough
+	 * precision to represent the range accurately.
+	 */
+	UFUNCTION(BlueprintCallable, Category="Media|MediaPlayer")
+	MEDIAASSETS_API bool SetPlaybackTimeRange(FFloatInterval InTimeRange);
+
 	/**
 	 * Clean up before this object is destroyed. Normally you would not need to call this,
 	 * but if you want to clean up resources before garbage collection (e.g. you are in the editor)
@@ -1068,7 +1150,7 @@ public:
 	/**
 	 * Templated version of GetMediaInfo.
 	 * No need to deal with variants.
-	 * 
+	 *
 	 * @param	T			Type of the information.
 	 * @param	Result		If the information is found and is the correct type, then this will be
 	 *						set to its value. It will not be set otherwise.
@@ -1103,7 +1185,7 @@ public:
 	 */
 	UFUNCTION(BlueprintCallable, Category = "Media|Metadata")
 	MEDIAASSETS_API TMap<FString, FMediaMetadataItemsBPT> GetMediaMetadataItems() const;
-	
+
 	/**
 	 * Get the media player facade that manages low-level media players
 	 *
@@ -1117,7 +1199,7 @@ public:
 	 */
 	MEDIAASSETS_API void RegisterWithMediaModule();
 
-	/** 
+	/**
 	 * When the player goes out of scope, make sure to clean up the clock sink
 	 */
 	MEDIAASSETS_API void UnregisterWithMediaModule();
diff --git a/Engine/Source/Runtime/MediaUtils/Private/MediaPlayerFacade.cpp b/Engine/Source/Runtime/MediaUtils/Private/MediaPlayerFacade.cpp
--- a/Engine/Source/Runtime/MediaUtils/Private/MediaPlayerFacade.cpp
+++ b/Engine/Source/Runtime/MediaUtils/Private/MediaPlayerFacade.cpp
@@ -260,6 +260,13 @@ bool FMediaPlayerFacade::CanSeek() const
 }
 
 
+bool FMediaPlayerFacade::SupportsPlaybackTimeRange() const
+{
+	TSharedPtr<IMediaPlayer, ESPMode::ThreadSafe> CurrentPlayer(Player);
+	return CurrentPlayer.IsValid() ? CurrentPlayer->GetControls().CanControl(EMediaControl::PlaybackRange) : false;
+}
+
+
 void FMediaPlayerFacade::Close()
 {
 	SCOPE_CYCLE_COUNTER(STAT_MediaUtils_FacadeClose);
@@ -1570,6 +1577,19 @@ bool FMediaPlayerFacade::SupportsRate(float Rate, bool Unthinned) const
 	return Player.IsValid() && Player->GetControls().GetSupportedRates(Thinning).Contains(Rate);
 }
 
+TRange<FTimespan> FMediaPlayerFacade::GetPlaybackTimeRange(EMediaTimeRangeType InRangeToGet) const
+{
+	TSharedPtr<IMediaPlayer, ESPMode::ThreadSafe> CurrentPlayer(Player);
+	return CurrentPlayer.IsValid() ? CurrentPlayer->GetControls().GetPlaybackTimeRange(InRangeToGet) : TRange<FTimespan>();
+}
+
+bool FMediaPlayerFacade::SetPlaybackTimeRange(const TRange<FTimespan>& InTimeRange)
+{
+	TSharedPtr<IMediaPlayer, ESPMode::ThreadSafe> CurrentPlayer(Player);
+	return CurrentPlayer.IsValid() ? CurrentPlayer->GetControls().SetPlaybackTimeRange(InTimeRange) : false;
+}
+
+
 void FMediaPlayerFacade::SetLastAudioRenderedSampleTime(FTimespan SampleTime)
 {
 	FScopeLock Lock(&LastTimeValuesCS);
diff --git a/Engine/Source/Runtime/MediaUtils/Public/MediaPlayerFacade.h b/Engine/Source/Runtime/MediaUtils/Public/MediaPlayerFacade.h
--- a/Engine/Source/Runtime/MediaUtils/Public/MediaPlayerFacade.h
+++ b/Engine/Source/Runtime/MediaUtils/Public/MediaPlayerFacade.h
@@ -42,6 +42,7 @@ enum class EMediaEvent;
 enum class EMediaCacheState;
 enum class EMediaThreads;
 enum class EMediaTrackType;
+enum class EMediaTimeRangeType;
 
 struct FMediaAudioTrackFormat;
 struct FMediaPlayerOptions;
@@ -177,6 +178,14 @@ public:
 	 */
 	MEDIAUTILS_API bool CanSeek() const;
 
+	/**
+	 * Check whether the player supports playing back of range within the media.
+	 *
+	 * @return true if playing back a range is supported, false otherwise.
+	 * @see GetPlaybackTimeRange, SetPlaybackTimeRange
+	 */
+	MEDIAUTILS_API bool SupportsPlaybackTimeRange() const;
+
 	/**
 	 * Close the currently open media, if any.
 	 */
@@ -238,7 +247,7 @@ public:
 
 	/**
 	 * Get information about the media that is playing.
-	 * 
+	 *
 	 * @param	InfoName		Name of the information we want.
 	 * @returns					Requested information, or empty if not available.
 	 * @see						UMediaPlayer::GetMediaInfo.
@@ -332,6 +341,32 @@ public:
 	 */
 	MEDIAUTILS_API TRangeSet<float> GetSupportedRates(bool Unthinned = true) const;
 
+	/**
+	 * Returns the current playback range of the media.
+	 * If playing back a range is not supported, the range returned will be equal
+	 * to [ 0, GetDuration() ).
+	 * The media may have an implicit default range provided by the container format
+	 * or other means without having called SetPlaybackTimeRange().
+	 * The media may have internal time values not starting at 0, which are
+	 * conveyed by the range.
+	 * Since the range may be only a portion of the media, the duration of the
+	 * returned range may be less than the media overall duration returned by
+	 * GetDuration().
+	 * For live video streams the range may change dynamically as new content
+	 * becomes available and old content falls off the timeline.
+	 *
+	 * @param InRangeToGet The type of range to get.
+	 *                     `Absolute` returns the media's smallest and largest timeline values.
+	 *                       Unless continuously changing in a Live stream this is usually the
+	 *                       same as [ 0, GetDuration() ]. The base time does not have to be
+	 *                       zero though.
+	 *                     `Current` returns the currently set range, which is a subset of the
+	 *                       absolute range.
+	 * @return The playback range as queried for or an empty range if there is no open player.
+	 * @see SupportsPlaybackTimeRange, SetPlaybackTimeRange, GetDuration
+	 */
+	MEDIAUTILS_API TRange<FTimespan> GetPlaybackTimeRange(EMediaTimeRangeType InRangeToGet) const;
+
 	/**
 	 * Get the media's current playback time.
 	 *
@@ -351,7 +386,7 @@ public:
 	 * Get the media's current playback time stamp in a "display" version
 	 *
 	 * @return Playback time stamp.
-	 * 
+	 *
 	 * @note The timestamp returned here will reflect a user-logic oriented version.
 	 *       (e.g. during seeks this will return the seek target rather than the last valid frame still displayed)
 	 */
@@ -641,6 +676,21 @@ public:
 	 */
 	MEDIAUTILS_API bool SetRate(float Rate);
 
+	/**
+	 * Sets a new media playback range.
+	 * Has an effect only if SupportsPlaybackTimeRange() returns true and the media supports it.
+	 * A live stream cannot be constrained to a range.
+	 * The range will be clamped if necessary to be within the media's absolute time range.
+	 * Changing the time range may trigger an implicit Seek() depending on where the current
+	 * playback position is located with regard to the new range.
+	 * Unless prevented by the media a playback range can be cleared by passing an empty range.
+	 *
+	 * @param InTimeRange The new playback range to set.
+	 * @return true if successful, false otherwise.
+	 * @see SupportsPlaybackTimeRange, GetPlaybackTimeRange
+	 */
+	MEDIAUTILS_API bool SetPlaybackTimeRange(const TRange<FTimespan>& InTimeRange);
+
 	/**
 	 * Changes the media's native volume.
 	 *
