Remove FRHICommandListExecutor::GetImmediateCommandList() in Media modules
 - Replaced with FRHICommandListImmediate::Get(), which includes a "check(IsInRenderingThread())".
 - Prefer to pass the FRHICommandListImmediate rather than grabbing it via the global scope function where possible. This makes it clearer which functions are always called by the rendering thread.

Also replace some direct use of "GDynamicRHI->" with the equivalent RHICmdList functions.

diff --git a/Engine/Plugins/Experimental/Avalanche/Source/AvalancheMedia/Private/Broadcast/OutputDevices/AvaBroadcastDisplayMediaCapture.cpp b/Engine/Plugins/Experimental/Avalanche/Source/AvalancheMedia/Private/Broadcast/OutputDevices/AvaBroadcastDisplayMediaCapture.cpp
--- a/Engine/Plugins/Experimental/Avalanche/Source/AvalancheMedia/Private/Broadcast/OutputDevices/AvaBroadcastDisplayMediaCapture.cpp
+++ b/Engine/Plugins/Experimental/Avalanche/Source/AvalancheMedia/Private/Broadcast/OutputDevices/AvaBroadcastDisplayMediaCapture.cpp
@@ -113,7 +113,7 @@ UAvaBroadcastDisplayMediaCapture::~UAvaBroadcastDisplayMediaCapture()
 	delete CaptureInstance;
 }
 
-void UAvaBroadcastDisplayMediaCapture::OnRHIResourceCaptured_RenderingThread(const FCaptureBaseData& InBaseData, TSharedPtr<FMediaCaptureUserData, ESPMode::ThreadSafe> InUserData, FTextureRHIRef InTexture)
+void UAvaBroadcastDisplayMediaCapture::OnRHIResourceCaptured_RenderingThread(FRHICommandListImmediate& RHICmdList, const FCaptureBaseData& InBaseData, TSharedPtr<FMediaCaptureUserData, ESPMode::ThreadSafe> InUserData, FTextureRHIRef InTexture)
 {
 	FScopeLock ScopeLock(&CaptureInstanceCriticalSection);
 	if (CaptureInstance)
@@ -137,7 +137,6 @@ void UAvaBroadcastDisplayMediaCapture::OnRHIResourceCaptured_RenderingThread(con
 
 		if (Target.IsValid())
 		{
-			FRHICommandListImmediate& RHICmdList = FRHICommandListExecutor::GetImmediateCommandList();
 			UE::AvaBroadcastRenderTargetMediaUtils::CopyTexture(RHICmdList, InTexture, Target);
 		}
 	}
diff --git a/Engine/Plugins/Experimental/Avalanche/Source/AvalancheMedia/Private/Broadcast/OutputDevices/AvaBroadcastDisplayMediaCapture.h b/Engine/Plugins/Experimental/Avalanche/Source/AvalancheMedia/Private/Broadcast/OutputDevices/AvaBroadcastDisplayMediaCapture.h
--- a/Engine/Plugins/Experimental/Avalanche/Source/AvalancheMedia/Private/Broadcast/OutputDevices/AvaBroadcastDisplayMediaCapture.h
+++ b/Engine/Plugins/Experimental/Avalanche/Source/AvalancheMedia/Private/Broadcast/OutputDevices/AvaBroadcastDisplayMediaCapture.h
@@ -18,7 +18,7 @@ public:
 
 protected:
 	virtual bool ShouldCaptureRHIResource() const override { return true; }
-	virtual void OnRHIResourceCaptured_RenderingThread(const FCaptureBaseData& InBaseData, TSharedPtr<FMediaCaptureUserData, ESPMode::ThreadSafe> InUserData, FTextureRHIRef InTexture) override;
+	virtual void OnRHIResourceCaptured_RenderingThread(FRHICommandListImmediate& RHICmdList, const FCaptureBaseData& InBaseData, TSharedPtr<FMediaCaptureUserData, ESPMode::ThreadSafe> InUserData, FTextureRHIRef InTexture) override;
 
 	virtual bool InitializeCapture() override;
 	virtual bool PostInitializeCaptureViewport(TSharedPtr<FSceneViewport>& InSceneViewport) override;
diff --git a/Engine/Plugins/Experimental/Avalanche/Source/AvalancheMedia/Private/Broadcast/OutputDevices/AvaBroadcastRenderTargetMediaCapture.cpp b/Engine/Plugins/Experimental/Avalanche/Source/AvalancheMedia/Private/Broadcast/OutputDevices/AvaBroadcastRenderTargetMediaCapture.cpp
--- a/Engine/Plugins/Experimental/Avalanche/Source/AvalancheMedia/Private/Broadcast/OutputDevices/AvaBroadcastRenderTargetMediaCapture.cpp
+++ b/Engine/Plugins/Experimental/Avalanche/Source/AvalancheMedia/Private/Broadcast/OutputDevices/AvaBroadcastRenderTargetMediaCapture.cpp
@@ -12,7 +12,7 @@
 
 DEFINE_LOG_CATEGORY(LogAvaBroadcastRenderTargetMedia);
 
-void UAvaBroadcastRenderTargetMediaCapture::OnRHIResourceCaptured_RenderingThread(const FCaptureBaseData& InBaseData, TSharedPtr<FMediaCaptureUserData, ESPMode::ThreadSafe> InUserData, FTextureRHIRef InTexture)
+void UAvaBroadcastRenderTargetMediaCapture::OnRHIResourceCaptured_RenderingThread(FRHICommandListImmediate& RHICmdList, const FCaptureBaseData& InBaseData, TSharedPtr<FMediaCaptureUserData, ESPMode::ThreadSafe> InUserData, FTextureRHIRef InTexture)
 {
 	FScopeLock ScopeLock(&RenderTargetCriticalSection);
 
@@ -25,8 +25,6 @@ void UAvaBroadcastRenderTargetMediaCapture::OnRHIResourceCaptured_RenderingThrea
 
 	if (RenderTargetRHIRef.IsValid())
 	{
-		FRHICommandListImmediate& RHICmdList = FRHICommandListExecutor::GetImmediateCommandList();
-
 		// Note: we possibly don't want to undo this for RGB8_SRGB texture. Need to test.
 		constexpr bool bSRGBToLinear = true;
 
diff --git a/Engine/Plugins/Experimental/Avalanche/Source/AvalancheMedia/Private/Broadcast/OutputDevices/AvaBroadcastRenderTargetMediaCapture.h b/Engine/Plugins/Experimental/Avalanche/Source/AvalancheMedia/Private/Broadcast/OutputDevices/AvaBroadcastRenderTargetMediaCapture.h
--- a/Engine/Plugins/Experimental/Avalanche/Source/AvalancheMedia/Private/Broadcast/OutputDevices/AvaBroadcastRenderTargetMediaCapture.h
+++ b/Engine/Plugins/Experimental/Avalanche/Source/AvalancheMedia/Private/Broadcast/OutputDevices/AvaBroadcastRenderTargetMediaCapture.h
@@ -15,7 +15,7 @@ class UAvaBroadcastRenderTargetMediaCapture : public UMediaCapture
 
 protected:
 	virtual bool ShouldCaptureRHIResource() const override { return true; }
-	virtual void OnRHIResourceCaptured_RenderingThread(const FCaptureBaseData& InBaseData, TSharedPtr<FMediaCaptureUserData, ESPMode::ThreadSafe> InUserData, FTextureRHIRef InTexture) override;
+	virtual void OnRHIResourceCaptured_RenderingThread(FRHICommandListImmediate& RHICmdList, const FCaptureBaseData& InBaseData, TSharedPtr<FMediaCaptureUserData, ESPMode::ThreadSafe> InUserData, FTextureRHIRef InTexture) override;
 
 	virtual bool InitializeCapture() override;
 	virtual bool PostInitializeCaptureViewport(TSharedPtr<FSceneViewport>& InSceneViewport) override;
diff --git a/Engine/Plugins/Experimental/PixelStreamingPlayer/Source/PixelStreamingPlayer/Private/PixelStreamingMediaTexture.cpp b/Engine/Plugins/Experimental/PixelStreamingPlayer/Source/PixelStreamingPlayer/Private/PixelStreamingMediaTexture.cpp
--- a/Engine/Plugins/Experimental/PixelStreamingPlayer/Source/PixelStreamingPlayer/Private/PixelStreamingMediaTexture.cpp
+++ b/Engine/Plugins/Experimental/PixelStreamingPlayer/Source/PixelStreamingPlayer/Private/PixelStreamingMediaTexture.cpp
@@ -53,7 +53,7 @@ void UPixelStreamingMediaTexture::OnFrame(FTextureRHIRef Frame)
 	AsyncTask(ENamedThreads::ActualRenderingThread, [this, Frame]() {
 		FScopeLock Lock(&RenderSyncContext);
 
-		FRHICommandListImmediate& RHICmdList = FRHICommandListExecutor::GetImmediateCommandList();
+		FRHICommandListImmediate& RHICmdList = FRHICommandListImmediate::Get();
 		UpdateTextureReference(RHICmdList, Frame);
 	});
 }
@@ -72,12 +72,12 @@ void UPixelStreamingMediaTexture::InitializeResources()
 				.SetFlags(ETextureCreateFlags::Dynamic | ETextureCreateFlags::ShaderResource | TexCreate_RenderTargetable)
 				.SetInitialState(ERHIAccess::SRVMask);
 	
-		RenderableTexture = RHICreateTexture(RenderTargetTextureDesc);
+		RenderableTexture = RHICmdList.CreateTexture(RenderTargetTextureDesc);
 		ShaderTexture2D = RenderableTexture;
 	
 		CurrentResource->TextureRHI = RenderableTexture;
 
-		RHIUpdateTextureReference(TextureReference.TextureReferenceRHI, CurrentResource->TextureRHI);
+		RHICmdList.UpdateTextureReference(TextureReference.TextureReferenceRHI, CurrentResource->TextureRHI);
 	});
 }
 
@@ -88,7 +88,7 @@ void UPixelStreamingMediaTexture::UpdateTextureReference(FRHICommandList& RHICmd
 		if (Reference.IsValid() && CurrentResource->TextureRHI != Reference)
 		{
 			CurrentResource->TextureRHI = Reference;
-			RHIUpdateTextureReference(TextureReference.TextureReferenceRHI, CurrentResource->TextureRHI);
+			RHICmdList.UpdateTextureReference(TextureReference.TextureReferenceRHI, CurrentResource->TextureRHI);
 		}
 		else if (!Reference.IsValid())
 		{
diff --git a/Engine/Plugins/Media/AjaMedia/Source/AjaMediaOutput/Private/AjaMediaCapture.cpp b/Engine/Plugins/Media/AjaMedia/Source/AjaMediaOutput/Private/AjaMediaCapture.cpp
--- a/Engine/Plugins/Media/AjaMedia/Source/AjaMediaOutput/Private/AjaMediaCapture.cpp
+++ b/Engine/Plugins/Media/AjaMedia/Source/AjaMediaOutput/Private/AjaMediaCapture.cpp
@@ -598,7 +598,7 @@ void UAjaMediaCapture::OnFrameCaptured_RenderingThread(const FCaptureBaseData& I
 	OnFrameCapturedInternal_AnyThread(InBaseData, InUserData, MoveTemp(InResourceData));
 }
 
-void UAjaMediaCapture::OnRHIResourceCaptured_RenderingThread(const FCaptureBaseData& InBaseData, TSharedPtr<FMediaCaptureUserData, ESPMode::ThreadSafe> InUserData, FTextureRHIRef InTexture)
+void UAjaMediaCapture::OnRHIResourceCaptured_RenderingThread(FRHICommandListImmediate& /*RHICmdList*/, const FCaptureBaseData& InBaseData, TSharedPtr<FMediaCaptureUserData, ESPMode::ThreadSafe> InUserData, FTextureRHIRef InTexture)
 {
 	OnRHIResourceCaptured_AnyThread(InBaseData, InUserData, InTexture);
 }
diff --git a/Engine/Plugins/Media/AjaMedia/Source/AjaMediaOutput/Public/AjaMediaCapture.h b/Engine/Plugins/Media/AjaMedia/Source/AjaMediaOutput/Public/AjaMediaCapture.h
--- a/Engine/Plugins/Media/AjaMedia/Source/AjaMediaOutput/Public/AjaMediaCapture.h
+++ b/Engine/Plugins/Media/AjaMedia/Source/AjaMediaOutput/Public/AjaMediaCapture.h
@@ -51,7 +51,7 @@ protected:
 	virtual bool ShouldCaptureRHIResource() const override;
 	
 	virtual void OnFrameCaptured_RenderingThread(const FCaptureBaseData& InBaseData, TSharedPtr<FMediaCaptureUserData, ESPMode::ThreadSafe> InUserData, void* InBuffer, int32 Width, int32 Height, int32 BytesPerRow) override;
-	virtual void OnRHIResourceCaptured_RenderingThread(const FCaptureBaseData& InBaseData, TSharedPtr<FMediaCaptureUserData, ESPMode::ThreadSafe> InUserData, FTextureRHIRef InTexture) override;
+	virtual void OnRHIResourceCaptured_RenderingThread(FRHICommandListImmediate& RHICmdList, const FCaptureBaseData& InBaseData, TSharedPtr<FMediaCaptureUserData, ESPMode::ThreadSafe> InUserData, FTextureRHIRef InTexture) override;
 	virtual void OnRHIResourceCaptured_AnyThread(const FCaptureBaseData& InBaseData, TSharedPtr<FMediaCaptureUserData, ESPMode::ThreadSafe> InUserData, FTextureRHIRef InTexture) override;
 	virtual void LockDMATexture_RenderThread(FTextureRHIRef InTexture) override;
 	virtual void UnlockDMATexture_RenderThread(FTextureRHIRef InTexture) override;
diff --git a/Engine/Plugins/Media/BinkMedia/Source/BinkMediaPlayer/Private/Assets/BinkMediaTextureResource.cpp b/Engine/Plugins/Media/BinkMedia/Source/BinkMediaPlayer/Private/Assets/BinkMediaTextureResource.cpp
--- a/Engine/Plugins/Media/BinkMedia/Source/BinkMediaPlayer/Private/Assets/BinkMediaTextureResource.cpp
+++ b/Engine/Plugins/Media/BinkMedia/Source/BinkMediaPlayer/Private/Assets/BinkMediaTextureResource.cpp
@@ -9,8 +9,10 @@
 #include "RenderingThread.h"
 #include "Rendering/SlateRenderer.h"
 
-void FBinkMediaTextureResource::InitRHI(FRHICommandListBase&) 
+void FBinkMediaTextureResource::InitRHI(FRHICommandListBase& RHICmdListBase) 
 {
+	FRHICommandListImmediate& RHICmdList = RHICmdListBase.GetAsImmediate();
+
 	int w = Owner->GetSurfaceWidth() > 0 ? Owner->GetSurfaceWidth() : 1;
 	int h = Owner->GetSurfaceHeight() > 0 ? Owner->GetSurfaceHeight() : 1;
 	// Enforce micro-tile restrictions for render targets.
@@ -46,7 +48,7 @@ void FBinkMediaTextureResource::InitRHI(FRHICommandListBase&)
 		.SetFlags(TexCreateFlags | ETextureCreateFlags::RenderTargetable | ETextureCreateFlags::ShaderResource)
 		.SetInitialState(ERHIAccess::SRVMask);
 
-	TextureRHI = RenderTargetTextureRHI = RHICreateTexture(Desc);
+	TextureRHI = RenderTargetTextureRHI = RHICmdList.CreateTexture(Desc);
 
 	// Don't bother updating if its not a valid video
 	if (Owner->GetSurfaceWidth() && Owner->GetSurfaceHeight()) 
@@ -64,9 +66,8 @@ void FBinkMediaTextureResource::InitRHI(FRHICommandListBase&)
 
 	SamplerStateRHI = RHICreateSamplerState(SamplerStateInitializer);
 
-	RHIUpdateTextureReference(Owner->TextureReference.TextureReferenceRHI, RenderTargetTextureRHI.GetReference());
+	RHICmdList.UpdateTextureReference(Owner->TextureReference.TextureReferenceRHI, RenderTargetTextureRHI.GetReference());
 
-	FRHICommandListImmediate& RHICmdList = FRHICommandListExecutor::GetImmediateCommandList();
 	{
 		FRHIRenderPassInfo RPInfo(TextureRHI, ERenderTargetActions::Clear_Store);
 		RHICmdList.Transition(FRHITransitionInfo(TextureRHI.GetReference(), ERHIAccess::Unknown, ERHIAccess::RTV));
diff --git a/Engine/Plugins/Media/BinkMedia/Source/BinkMediaPlayer/Private/binktexturesRHI.cpp b/Engine/Plugins/Media/BinkMedia/Source/BinkMediaPlayer/Private/binktexturesRHI.cpp
--- a/Engine/Plugins/Media/BinkMedia/Source/BinkMediaPlayer/Private/binktexturesRHI.cpp
+++ b/Engine/Plugins/Media/BinkMedia/Source/BinkMediaPlayer/Private/binktexturesRHI.cpp
@@ -143,7 +143,7 @@ static BINKTEXTURES * Create_textures(BINKSHADERS * pshaders, HBINK bink, void *
 	}
 #endif
 
-	FRHICommandListImmediate& RHICmdList = GetImmediateCommandList_ForRenderCommand();
+	FRHICommandListImmediate& RHICmdList = FRHICommandListImmediate::Get();
 	EPixelFormat format = PF_R8;
 #if PLATFORM_HAS_DIRECT_TEXTURE_MEMORY_ACCESS 
 	ETextureCreateFlags TexCreateFlags = TexCreate_Dynamic | TexCreate_CPUReadback | TexCreate_CPUWritable | TexCreate_OfflineProcessed | TexCreate_NoTiling;
@@ -167,7 +167,7 @@ static BINKTEXTURES * Create_textures(BINKSHADERS * pshaders, HBINK bink, void *
 		{
 			textures->Ytexture[i] = RHICreateTexture(YADesc);
 #if PLATFORM_HAS_DIRECT_TEXTURE_MEMORY_ACCESS 
-			bp_src->YPlane.Buffer = GDynamicRHI->LockTexture2D_RenderThread(RHICmdList, textures->Ytexture[i], 0, PLATFORM_DIRECT_TEXTURE_MEMORY_ACCESS_LOCK_MODE, bp_src->YPlane.BufferPitch, false, false);
+			bp_src->YPlane.Buffer = RHICmdList.LockTexture2D(textures->Ytexture[i], 0, PLATFORM_DIRECT_TEXTURE_MEMORY_ACCESS_LOCK_MODE, bp_src->YPlane.BufferPitch, false, false);
 			bp_src->YPlane.BufferPitch = bp_src->YPlane.BufferPitch ? bp_src->YPlane.BufferPitch : ((bb->YABufferWidth+255)&-256);
 #endif
 		}
@@ -176,7 +176,7 @@ static BINKTEXTURES * Create_textures(BINKSHADERS * pshaders, HBINK bink, void *
 		{
 			textures->cRtexture[i] = RHICreateTexture(cRcBDesc);
 #if PLATFORM_HAS_DIRECT_TEXTURE_MEMORY_ACCESS 
-			bp_src->cRPlane.Buffer = GDynamicRHI->LockTexture2D_RenderThread(RHICmdList, textures->cRtexture[i], 0, PLATFORM_DIRECT_TEXTURE_MEMORY_ACCESS_LOCK_MODE, bp_src->cRPlane.BufferPitch, false, false);
+			bp_src->cRPlane.Buffer = RHICmdList.LockTexture2D(textures->cRtexture[i], 0, PLATFORM_DIRECT_TEXTURE_MEMORY_ACCESS_LOCK_MODE, bp_src->cRPlane.BufferPitch, false, false);
 			bp_src->cRPlane.BufferPitch = bp_src->cRPlane.BufferPitch ? bp_src->cRPlane.BufferPitch : ((bb->cRcBBufferWidth+255)&-256);
 #endif
 		}
@@ -185,7 +185,7 @@ static BINKTEXTURES * Create_textures(BINKSHADERS * pshaders, HBINK bink, void *
 		{
 			textures->cBtexture[i] = RHICreateTexture(cRcBDesc);
 #if PLATFORM_HAS_DIRECT_TEXTURE_MEMORY_ACCESS 
-			bp_src->cBPlane.Buffer = GDynamicRHI->LockTexture2D_RenderThread(RHICmdList, textures->cBtexture[i], 0, PLATFORM_DIRECT_TEXTURE_MEMORY_ACCESS_LOCK_MODE, bp_src->cBPlane.BufferPitch, false, false);
+			bp_src->cBPlane.Buffer = RHICmdList.LockTexture2D(textures->cBtexture[i], 0, PLATFORM_DIRECT_TEXTURE_MEMORY_ACCESS_LOCK_MODE, bp_src->cBPlane.BufferPitch, false, false);
 			bp_src->cBPlane.BufferPitch = bp_src->cBPlane.BufferPitch ? bp_src->cBPlane.BufferPitch : ((bb->cRcBBufferWidth+255)&-256);
 #endif
 		}
@@ -194,7 +194,7 @@ static BINKTEXTURES * Create_textures(BINKSHADERS * pshaders, HBINK bink, void *
 		{
 			textures->Atexture[i] = RHICreateTexture(YADesc);
 #if PLATFORM_HAS_DIRECT_TEXTURE_MEMORY_ACCESS 
-			bp_src->APlane.Buffer = GDynamicRHI->LockTexture2D_RenderThread(RHICmdList, textures->Atexture[i], 0, PLATFORM_DIRECT_TEXTURE_MEMORY_ACCESS_LOCK_MODE, bp_src->APlane.BufferPitch, false, false);
+			bp_src->APlane.Buffer = RHICmdList.LockTexture2D(textures->Atexture[i], 0, PLATFORM_DIRECT_TEXTURE_MEMORY_ACCESS_LOCK_MODE, bp_src->APlane.BufferPitch, false, false);
 			bp_src->APlane.BufferPitch = bp_src->APlane.BufferPitch ? bp_src->APlane.BufferPitch : ((bb->YABufferWidth+255)&-256);
 #endif
 		}
@@ -203,7 +203,7 @@ static BINKTEXTURES * Create_textures(BINKSHADERS * pshaders, HBINK bink, void *
 		{
 			textures->Htexture[i] = RHICreateTexture(YADesc);
 #if PLATFORM_HAS_DIRECT_TEXTURE_MEMORY_ACCESS 
-			bp_src->HPlane.Buffer = GDynamicRHI->LockTexture2D_RenderThread(RHICmdList, textures->Htexture[i], 0, PLATFORM_DIRECT_TEXTURE_MEMORY_ACCESS_LOCK_MODE, bp_src->HPlane.BufferPitch, false, false);
+			bp_src->HPlane.Buffer = RHICmdList.LockTexture2D(textures->Htexture[i], 0, PLATFORM_DIRECT_TEXTURE_MEMORY_ACCESS_LOCK_MODE, bp_src->HPlane.BufferPitch, false, false);
 			bp_src->HPlane.BufferPitch = bp_src->HPlane.BufferPitch ? bp_src->HPlane.BufferPitch : ((bb->YABufferWidth+255)&-256);
 #endif
 		}
@@ -234,7 +234,7 @@ static BINKTEXTURES * Create_textures(BINKSHADERS * pshaders, HBINK bink, void *
 
 static void Free_textures(BINKTEXTURES* ptextures)
 {
-	FRHICommandListImmediate& RHICmdList = GetImmediateCommandList_ForRenderCommand();
+	FRHICommandListImmediate& RHICmdList = FRHICommandListImmediate::Get();
 	BINKTEXTURESRHI* textures = (BINKTEXTURESRHI*)ptextures;
 	BINKFRAMEBUFFERS *bb = &textures->bink_buffers;
 
@@ -249,11 +249,11 @@ static void Free_textures(BINKTEXTURES* ptextures)
 	{
 		BINKFRAMEPLANESET *bp_src = &bb->Frames[i];
 #if PLATFORM_HAS_DIRECT_TEXTURE_MEMORY_ACCESS 
-		if (bp_src->YPlane.Buffer)  GDynamicRHI->UnlockTexture2D_RenderThread(RHICmdList, textures->Ytexture[i], 0, false);
-		if (bp_src->cRPlane.Buffer) GDynamicRHI->UnlockTexture2D_RenderThread(RHICmdList, textures->cRtexture[i], 0, false);
-		if (bp_src->cBPlane.Buffer) GDynamicRHI->UnlockTexture2D_RenderThread(RHICmdList, textures->cBtexture[i], 0, false);
-		if (bp_src->APlane.Buffer)  GDynamicRHI->UnlockTexture2D_RenderThread(RHICmdList, textures->Atexture[i], 0, false);
-		if (bp_src->HPlane.Buffer)  GDynamicRHI->UnlockTexture2D_RenderThread(RHICmdList, textures->Htexture[i], 0, false);
+		if (bp_src->YPlane.Buffer)  RHICmdList.UnlockTexture2D(textures->Ytexture[i], 0, false);
+		if (bp_src->cRPlane.Buffer) RHICmdList.UnlockTexture2D(textures->cRtexture[i], 0, false);
+		if (bp_src->cBPlane.Buffer) RHICmdList.UnlockTexture2D(textures->cBtexture[i], 0, false);
+		if (bp_src->APlane.Buffer)  RHICmdList.UnlockTexture2D(textures->Atexture[i], 0, false);
+		if (bp_src->HPlane.Buffer)  RHICmdList.UnlockTexture2D(textures->Htexture[i], 0, false);
 #endif
 		if (textures->Ytexture[i].IsValid())  textures->Ytexture[i].SafeRelease();
 		if (textures->cRtexture[i].IsValid()) textures->cRtexture[i].SafeRelease();
@@ -272,7 +272,6 @@ static void Start_texture_update( BINKTEXTURES * ptextures )
 
 static void Finish_texture_update( BINKTEXTURES * ptextures )
 {
-	FRHICommandListImmediate& RHICmdList = GetImmediateCommandList_ForRenderCommand();
 	BINKTEXTURESRHI * textures = (BINKTEXTURESRHI*)ptextures;
 	BINKFRAMEBUFFERS *bb = &textures->bink_buffers;
 	int frame_num = bb->FrameNum;
@@ -332,8 +331,7 @@ static void Draw_textures(BINKTEXTURES* ptextures, BINKSHADERS* pshaders, void*
 	int hasAPlane = bp_src->APlane.Allocate;
 	int hasHPlane = bp_src->HPlane.Allocate;
 
-	FRHICommandListImmediate& RHICmdList = GetImmediateCommandList_ForRenderCommand();
-	FRDGBuilder BinkGraphBuilder(RHICmdList);
+	FRDGBuilder BinkGraphBuilder(FRHICommandListImmediate::Get());
 
 	FBinkParameters *consts = BinkGraphBuilder.AllocParameters<FBinkParameters>();
 
@@ -482,11 +480,11 @@ static void Draw_textures(BINKTEXTURES* ptextures, BINKSHADERS* pshaders, void*
 				FUpdateTextureRegion2D region_YAH(0, 0, 0, 0, bb->YABufferWidth, bb->YABufferHeight);
 				FUpdateTextureRegion2D region_cRcB(0, 0, 0, 0, bb->cRcBBufferWidth, bb->cRcBBufferHeight);
 
-				GDynamicRHI->RHIUpdateTexture2D(RHICmdList, textures->Ytexture[frame_num], 0, region_YAH, bp_src->YPlane.BufferPitch, (uint8*)bp_src->YPlane.Buffer);
-				GDynamicRHI->RHIUpdateTexture2D(RHICmdList, textures->cRtexture[frame_num], 0, region_cRcB, bp_src->cRPlane.BufferPitch, (uint8*)bp_src->cRPlane.Buffer);
-				GDynamicRHI->RHIUpdateTexture2D(RHICmdList, textures->cBtexture[frame_num], 0, region_cRcB, bp_src->cBPlane.BufferPitch, (uint8*)bp_src->cBPlane.Buffer);
-				if (hasAPlane) GDynamicRHI->RHIUpdateTexture2D(RHICmdList, textures->Atexture[frame_num], 0, region_YAH, bp_src->APlane.BufferPitch, (uint8*)bp_src->APlane.Buffer);
-				if (hasHPlane) GDynamicRHI->RHIUpdateTexture2D(RHICmdList, textures->Htexture[frame_num], 0, region_YAH, bp_src->HPlane.BufferPitch, (uint8*)bp_src->HPlane.Buffer);
+				RHICmdList.UpdateTexture2D(textures->Ytexture[frame_num], 0, region_YAH, bp_src->YPlane.BufferPitch, (uint8*)bp_src->YPlane.Buffer);
+				RHICmdList.UpdateTexture2D(textures->cRtexture[frame_num], 0, region_cRcB, bp_src->cRPlane.BufferPitch, (uint8*)bp_src->cRPlane.Buffer);
+				RHICmdList.UpdateTexture2D(textures->cBtexture[frame_num], 0, region_cRcB, bp_src->cBPlane.BufferPitch, (uint8*)bp_src->cBPlane.Buffer);
+				if (hasAPlane) RHICmdList.UpdateTexture2D(textures->Atexture[frame_num], 0, region_YAH, bp_src->APlane.BufferPitch, (uint8*)bp_src->APlane.Buffer);
+				if (hasHPlane) RHICmdList.UpdateTexture2D(textures->Htexture[frame_num], 0, region_YAH, bp_src->HPlane.BufferPitch, (uint8*)bp_src->HPlane.Buffer);
 			}
 
 			FGraphicsPipelineStateInitializer GraphicsPSOInit;
diff --git a/Engine/Plugins/Media/BlackmagicMedia/Source/BlackmagicMediaOutput/Private/BlackmagicMediaCapture.cpp b/Engine/Plugins/Media/BlackmagicMedia/Source/BlackmagicMediaOutput/Private/BlackmagicMediaCapture.cpp
--- a/Engine/Plugins/Media/BlackmagicMedia/Source/BlackmagicMediaOutput/Private/BlackmagicMediaCapture.cpp
+++ b/Engine/Plugins/Media/BlackmagicMedia/Source/BlackmagicMediaOutput/Private/BlackmagicMediaCapture.cpp
@@ -824,7 +824,7 @@ void UBlackmagicMediaCapture::OnFrameCaptured_AnyThread(const FCaptureBaseData&
 	OnFrameCapturedInternal_AnyThread(InBaseData, InUserData, InResourceData);
 }
 
-void UBlackmagicMediaCapture::OnRHIResourceCaptured_RenderingThread(const FCaptureBaseData& InBaseData, TSharedPtr<FMediaCaptureUserData, ESPMode::ThreadSafe> InUserData, FTextureRHIRef InTexture)
+void UBlackmagicMediaCapture::OnRHIResourceCaptured_RenderingThread(FRHICommandListImmediate& /*RHICmdList*/, const FCaptureBaseData& InBaseData, TSharedPtr<FMediaCaptureUserData, ESPMode::ThreadSafe> InUserData, FTextureRHIRef InTexture)
 {
 	TRACE_CPUPROFILER_EVENT_SCOPE(UBlackmagicMediaCapture::OnRHIResourceCaptured_RenderingThread);
 	OnRHIResourceCaptured_AnyThread(InBaseData, InUserData, InTexture);
diff --git a/Engine/Plugins/Media/BlackmagicMedia/Source/BlackmagicMediaOutput/Public/BlackmagicMediaCapture.h b/Engine/Plugins/Media/BlackmagicMedia/Source/BlackmagicMediaOutput/Public/BlackmagicMediaCapture.h
--- a/Engine/Plugins/Media/BlackmagicMedia/Source/BlackmagicMediaOutput/Public/BlackmagicMediaCapture.h
+++ b/Engine/Plugins/Media/BlackmagicMedia/Source/BlackmagicMediaOutput/Public/BlackmagicMediaCapture.h
@@ -54,7 +54,7 @@ protected:
 	virtual void OnFrameCaptured_RenderingThread(const FCaptureBaseData& InBaseData, TSharedPtr<FMediaCaptureUserData, ESPMode::ThreadSafe> InUserData, void* InBuffer, int32 Width, int32 Height, int32 BytesPerRow) override;
 	virtual void OnFrameCaptured_AnyThread(const FCaptureBaseData& InBaseData, TSharedPtr<FMediaCaptureUserData, ESPMode::ThreadSafe> InUserData, const FMediaCaptureResourceData& InResourceData) override;
 	virtual void OnRHIResourceCaptured_AnyThread(const FCaptureBaseData& InBaseData, TSharedPtr<FMediaCaptureUserData, ESPMode::ThreadSafe> InUserData, FTextureRHIRef InTexture) override;
-	virtual void OnRHIResourceCaptured_RenderingThread(const FCaptureBaseData& InBaseData, TSharedPtr<FMediaCaptureUserData, ESPMode::ThreadSafe> InUserData, FTextureRHIRef InTexture) override;
+	virtual void OnRHIResourceCaptured_RenderingThread(FRHICommandListImmediate& RHICmdList, const FCaptureBaseData& InBaseData, TSharedPtr<FMediaCaptureUserData, ESPMode::ThreadSafe> InUserData, FTextureRHIRef InTexture) override;
 	virtual void LockDMATexture_RenderThread(FTextureRHIRef InTexture) override;
 	virtual void UnlockDMATexture_RenderThread(FTextureRHIRef InTexture) override;
 	virtual bool SupportsAnyThreadCapture() const override;
diff --git a/Engine/Plugins/Media/ElectraUtil/Source/ElectraSamples/Private/ElectraTextureSampleAVF.cpp b/Engine/Plugins/Media/ElectraUtil/Source/ElectraSamples/Private/ElectraTextureSampleAVF.cpp
--- a/Engine/Plugins/Media/ElectraUtil/Source/ElectraSamples/Private/ElectraTextureSampleAVF.cpp
+++ b/Engine/Plugins/Media/ElectraUtil/Source/ElectraSamples/Private/ElectraTextureSampleAVF.cpp
@@ -445,7 +445,7 @@ uint32 FElectraTextureSample::GetConverterInfoFlags() const
 	return ConverterInfoFlags_Default;
 }
 
-bool FElectraTextureSample::Convert(FTextureRHIRef & InDstTexture, const FConversionHints & Hints)
+bool FElectraTextureSample::Convert(FRHICommandListImmediate& RHICmdList, FTextureRHIRef& InDstTexture, const FConversionHints& Hints)
 {
 	if (VideoDecoderOutput)
 	{
diff --git a/Engine/Plugins/Media/ElectraUtil/Source/ElectraSamples/Private/ElectraTextureSampleAndroid.cpp b/Engine/Plugins/Media/ElectraUtil/Source/ElectraSamples/Private/ElectraTextureSampleAndroid.cpp
--- a/Engine/Plugins/Media/ElectraUtil/Source/ElectraSamples/Private/ElectraTextureSampleAndroid.cpp
+++ b/Engine/Plugins/Media/ElectraUtil/Source/ElectraSamples/Private/ElectraTextureSampleAndroid.cpp
@@ -483,10 +483,8 @@ uint32 FElectraTextureSample::GetStride() const
 }
 
 
-bool FElectraTextureSample::Convert(FTextureRHIRef& InDstTexture, const FConversionHints& Hints)
+bool FElectraTextureSample::Convert(FRHICommandListImmediate& RHICmdList, FTextureRHIRef& InDstTexture, const FConversionHints& Hints)
 {
-	FRHICommandListImmediate& RHICmdList = FRHICommandListImmediate::Get();
-
 	if (GDynamicRHI->RHIIsRenderingSuspended())
 	{
 		return false;
@@ -504,7 +502,7 @@ bool FElectraTextureSample::Convert(FTextureRHIRef& InDstTexture, const FConvers
 			FRHITextureCreateDesc::Create2D(TEXT("FMediaTextureResource"), SampleDim, VideoDecoderOutputAndroid->GetFormat())
 			.SetFlags(TexCreate_Dynamic)
 			.SetInitialState(ERHIAccess::SRVMask);
-		InputTexture = RHICreateTexture(Desc);
+		InputTexture = RHICmdList.CreateTexture(Desc);
 		if (!InputTexture.IsValid())
 		{
 			return false;
@@ -512,7 +510,7 @@ bool FElectraTextureSample::Convert(FTextureRHIRef& InDstTexture, const FConvers
 
 		// copy sample data to input render target
 		FUpdateTextureRegion2D Region(0, 0, 0, 0, SampleDim.X, SampleDim.Y);
-		RHIUpdateTexture2D(InputTexture, 0, Region, GetStride(), (const uint8*)GetBuffer());
+		RHICmdList.UpdateTexture2D(InputTexture, 0, Region, GetStride(), (const uint8*)GetBuffer());
 	}
 	else
 	{
diff --git a/Engine/Plugins/Media/ElectraUtil/Source/ElectraSamples/Private/ElectraTextureSamplePC.cpp b/Engine/Plugins/Media/ElectraUtil/Source/ElectraSamples/Private/ElectraTextureSamplePC.cpp
--- a/Engine/Plugins/Media/ElectraUtil/Source/ElectraSamples/Private/ElectraTextureSamplePC.cpp
+++ b/Engine/Plugins/Media/ElectraUtil/Source/ElectraSamples/Private/ElectraTextureSamplePC.cpp
@@ -334,16 +334,13 @@ struct FRHICommandCopyResourceDX11 final : public FRHICommand<FRHICommandCopyRes
 /**
  * "Converter" for textures - here: a copy from the decoder owned texture (possibly in another device) into a RHI one (as prep for the real conversion to RGB etc.)
  */
-bool FElectraTextureSample::Convert(FTextureRHIRef& InDstTexture, const FConversionHints& Hints)
+bool FElectraTextureSample::Convert(FRHICommandListImmediate& RHICmdList, FTextureRHIRef& InDstTexture, const FConversionHints& Hints)
 {
 	LLM_SCOPE(ELLMTag::MediaStreaming);
 
-	FRHICommandListImmediate& RHICmdList = FRHICommandListImmediate::Get();
-
 	SCOPED_DRAW_EVENT(RHICmdList, WinMediaOutputConvertTexture);
 	SCOPED_GPU_STAT(RHICmdList, MediaWinDecoder_Convert);
 
-
 	bool bHasTexture = !!VideoDecoderOutputPC->GetTexture();
 
 	// Get actual sample dimensions
@@ -428,7 +425,7 @@ bool FElectraTextureSample::Convert(FTextureRHIRef& InDstTexture, const FConvers
 			FRHITextureCreateDesc::Create2D(TEXT("FElectraTextureSample"), Dim, Format)
 			.SetFlags(ETextureCreateFlags::Dynamic | ((bCanUseSRGB && IsOutputSrgb()) ? ETextureCreateFlags::SRGB : ETextureCreateFlags::None));
 
-		Texture = RHICreateTexture(Desc);
+		Texture = RHICmdList.CreateTexture(Desc);
 	}
 
 	uint64 SyncValue = 0;
diff --git a/Engine/Plugins/Media/ElectraUtil/Source/ElectraSamples/Public/Android/ElectraTextureSample.h b/Engine/Plugins/Media/ElectraUtil/Source/ElectraSamples/Public/Android/ElectraTextureSample.h
--- a/Engine/Plugins/Media/ElectraUtil/Source/ElectraSamples/Public/Android/ElectraTextureSample.h
+++ b/Engine/Plugins/Media/ElectraUtil/Source/ElectraSamples/Public/Android/ElectraTextureSample.h
@@ -77,7 +77,7 @@ private:
 	void InitializeTexture(EPixelFormat PixelFormat);
 	void SetupFromBuffer(const void* InBuffer, int32 InBufferSize);
 
-	virtual bool Convert(FTextureRHIRef& InDstTexture, const FConversionHints& Hints) override;
+	virtual bool Convert(FRHICommandListImmediate& RHICmdList, FTextureRHIRef& InDstTexture, const FConversionHints& Hints) override;
 	virtual uint32 GetConverterInfoFlags() const
 	{
 		return ConverterInfoFlags_Default;
diff --git a/Engine/Plugins/Media/ElectraUtil/Source/ElectraSamples/Public/Apple/ElectraTextureSample.h b/Engine/Plugins/Media/ElectraUtil/Source/ElectraSamples/Public/Apple/ElectraTextureSample.h
--- a/Engine/Plugins/Media/ElectraUtil/Source/ElectraSamples/Public/Apple/ElectraTextureSample.h
+++ b/Engine/Plugins/Media/ElectraUtil/Source/ElectraSamples/Public/Apple/ElectraTextureSample.h
@@ -58,7 +58,7 @@ private:
 	TWeakPtr<FElectraMediaTexConvApple, ESPMode::ThreadSafe> TexConv;
 
 	virtual uint32 GetConverterInfoFlags() const override;
-    virtual bool Convert(FTextureRHIRef & InDstTexture, const FConversionHints & Hints) override;
+    virtual bool Convert(FRHICommandListImmediate& RHICmdList, FTextureRHIRef& InDstTexture, const FConversionHints& Hints) override;
 };
 
 
diff --git a/Engine/Plugins/Media/ElectraUtil/Source/ElectraSamples/Public/Windows/ElectraTextureSample.h b/Engine/Plugins/Media/ElectraUtil/Source/ElectraSamples/Public/Windows/ElectraTextureSample.h
--- a/Engine/Plugins/Media/ElectraUtil/Source/ElectraSamples/Public/Windows/ElectraTextureSample.h
+++ b/Engine/Plugins/Media/ElectraUtil/Source/ElectraSamples/Public/Windows/ElectraTextureSample.h
@@ -56,7 +56,7 @@ public:
 private:
 	virtual float GetSampleDataScale(bool b10Bit) const override;
 
-	virtual bool Convert(FTextureRHIRef & InDstTexture, const FConversionHints & Hints) override;
+	virtual bool Convert(FRHICommandListImmediate& RHICmdList, FTextureRHIRef& InDstTexture, const FConversionHints& Hints) override;
 	virtual uint32 GetConverterInfoFlags() const
 	{
 		return ConverterInfoFlags_PreprocessOnly;
diff --git a/Engine/Plugins/Media/ImgMedia/Source/ImgMedia/Private/Readers/ExrImgMediaReaderGpu.cpp b/Engine/Plugins/Media/ImgMedia/Source/ImgMedia/Private/Readers/ExrImgMediaReaderGpu.cpp
--- a/Engine/Plugins/Media/ImgMedia/Source/ImgMedia/Private/Readers/ExrImgMediaReaderGpu.cpp
+++ b/Engine/Plugins/Media/ImgMedia/Source/ImgMedia/Private/Readers/ExrImgMediaReaderGpu.cpp
@@ -691,13 +691,13 @@ void FExrImgMediaReaderGpu::ReturnGpuBufferToPool(uint32 AllocSize, FStructuredB
 /* FExrMediaTextureSampleConverter implementation
  *****************************************************************************/
 
-bool FExrMediaTextureSampleConverter::Convert(FTextureRHIRef& InDstTexture, const FConversionHints& Hints)
+bool FExrMediaTextureSampleConverter::Convert(FRHICommandListImmediate& RHICmdList, FTextureRHIRef& InDstTexture, const FConversionHints& Hints)
 {
 	FScopeLock ScopeLock(&ConverterCallbacksCriticalSection);
 	bool bExecutionSuccessful = false;
 	if (ConvertExrBufferCallback.IsBound())
 	{
-		bExecutionSuccessful = ConvertExrBufferCallback.Execute(FRHICommandListExecutor::GetImmediateCommandList(), InDstTexture, MipBuffers, GetParams());
+		bExecutionSuccessful = ConvertExrBufferCallback.Execute(RHICmdList, InDstTexture, MipBuffers, GetParams());
 	}
 	return bExecutionSuccessful;
 }
diff --git a/Engine/Plugins/Media/ImgMedia/Source/ImgMedia/Private/Readers/ExrImgMediaReaderGpu.h b/Engine/Plugins/Media/ImgMedia/Source/ImgMedia/Private/Readers/ExrImgMediaReaderGpu.h
--- a/Engine/Plugins/Media/ImgMedia/Source/ImgMedia/Private/Readers/ExrImgMediaReaderGpu.h
+++ b/Engine/Plugins/Media/ImgMedia/Source/ImgMedia/Private/Readers/ExrImgMediaReaderGpu.h
@@ -185,7 +185,7 @@ class FExrMediaTextureSampleConverter: public IMediaTextureSampleConverter
 {
 
 public:
-	virtual bool Convert(FTextureRHIRef& InDstTexture, const FConversionHints& Hints) override;
+	virtual bool Convert(FRHICommandListImmediate& RHICmdList, FTextureRHIRef& InDstTexture, const FConversionHints& Hints) override;
 	virtual ~FExrMediaTextureSampleConverter() {};
 	
 	void AddCallback(FExrConvertBufferCallback&& Callback) 
diff --git a/Engine/Plugins/Media/MediaIOFramework/Source/MediaIOCore/Private/MediaCapture.cpp b/Engine/Plugins/Media/MediaIOFramework/Source/MediaIOCore/Private/MediaCapture.cpp
--- a/Engine/Plugins/Media/MediaIOFramework/Source/MediaIOCore/Private/MediaCapture.cpp
+++ b/Engine/Plugins/Media/MediaIOFramework/Source/MediaIOCore/Private/MediaCapture.cpp
@@ -1481,12 +1481,12 @@ bool UMediaCapture::ProcessReadyFrame_RenderThread(FRHICommandListImmediate& RHI
 				TRACE_CPUPROFILER_EVENT_SCOPE(UMediaCapture::RHIResourceCaptured);
 				TRACE_CPUPROFILER_EVENT_SCOPE_TEXT(*FString::Printf(TEXT("MediaCapture Output Frame %d"), ReadyFrame->CaptureBaseData.SourceFrameNumberRenderThread % 10));
 				SCOPE_CYCLE_COUNTER(STAT_MediaCapture_RenderThread_CaptureCallback)
-					InMediaCapture->OnRHIResourceCaptured_RenderingThread(ReadyFrame->CaptureBaseData, ReadyFrame->UserData, ReadyFrame->GetTextureResource());
+					InMediaCapture->OnRHIResourceCaptured_RenderingThread(RHICmdList, ReadyFrame->CaptureBaseData, ReadyFrame->UserData, ReadyFrame->GetTextureResource());
 			}
 			else
 			{
 				TRACE_CPUPROFILER_EVENT_SCOPE(UMediaCapture::RHIResourceCaptured);
-				InMediaCapture->OnRHIResourceCaptured_RenderingThread(ReadyFrame->CaptureBaseData, ReadyFrame->UserData, ReadyFrame->GetBufferResource());
+				InMediaCapture->OnRHIResourceCaptured_RenderingThread(RHICmdList, ReadyFrame->CaptureBaseData, ReadyFrame->UserData, ReadyFrame->GetBufferResource());
 			}
 
 			UE_LOG(LogMediaIOCore, Verbose, TEXT("[%s - %s] - Completed pending frame %d."), *InMediaCapture->MediaOutputName, *FThreadManager::GetThreadName(FPlatformTLS::GetCurrentThreadId()), ReadyFrame->FrameId);
diff --git a/Engine/Plugins/Media/MediaIOFramework/Source/MediaIOCore/Private/MediaCaptureHelper.cpp b/Engine/Plugins/Media/MediaIOFramework/Source/MediaIOCore/Private/MediaCaptureHelper.cpp
--- a/Engine/Plugins/Media/MediaIOFramework/Source/MediaIOCore/Private/MediaCaptureHelper.cpp
+++ b/Engine/Plugins/Media/MediaIOFramework/Source/MediaIOCore/Private/MediaCaptureHelper.cpp
@@ -212,7 +212,7 @@ void FMediaCaptureHelper::ExecuteSyncPointPass(FRHICommandListImmediate& RHICmdL
 	}
 }
 
-void FMediaCaptureHelper::OnReadbackComplete(FRHICommandList& RHICmdList, UMediaCapture* MediaCapture, TSharedPtr<UE::MediaCaptureData::FCaptureFrame> ReadyFrame)
+void FMediaCaptureHelper::OnReadbackComplete(UMediaCapture* MediaCapture, TSharedPtr<UE::MediaCaptureData::FCaptureFrame> ReadyFrame)
 {
 	UE_LOG(LogMediaIOCore, Verbose, TEXT("[%s - %s] - Processing pending frame %d"), *MediaCapture->GetMediaOutputName(), *FThreadManager::GetThreadName(FPlatformTLS::GetCurrentThreadId()), ReadyFrame->GetId());
 	TRACE_CPUPROFILER_EVENT_SCOPE(UMediaCapture::OnReadbackComplete);
@@ -229,6 +229,8 @@ void FMediaCaptureHelper::OnReadbackComplete(FRHICommandList& RHICmdList, UMedia
 		// Path where resource ready callback (readback / rhi capture) is on render thread (old method)
 		if (IsInRenderingThread())
 		{
+			FRHICommandListImmediate& RHICmdList = FRHICommandListImmediate::Get();
+
 			// Scoped gpu mask shouldn't be needed for readback since we specify the gpu mask used during copy when we lock
 			// Keeping it for old render thread path
 			FRHIGPUMask GPUMask;
@@ -255,14 +257,14 @@ void FMediaCaptureHelper::OnReadbackComplete(FRHICommandList& RHICmdList, UMedia
 
 					TRACE_CPUPROFILER_EVENT_SCOPE(UMediaCapture::RHIResourceCaptured);
 					TRACE_CPUPROFILER_EVENT_SCOPE_TEXT(*FString::Printf(TEXT("MediaCapture Output Frame %d"), ReadyFrame->CaptureBaseData.SourceFrameNumberRenderThread % 10));
-					SCOPE_CYCLE_COUNTER(STAT_MediaCaptureHelper_RenderThread_RHI_CaptureCallback)
+					SCOPE_CYCLE_COUNTER(STAT_MediaCaptureHelper_RenderThread_RHI_CaptureCallback);
 
-						MediaCapture->OnRHIResourceCaptured_RenderingThread(ReadyFrame->CaptureBaseData, ReadyFrame->UserData, ReadyFrame->GetTextureResource());
+					MediaCapture->OnRHIResourceCaptured_RenderingThread(RHICmdList, ReadyFrame->CaptureBaseData, ReadyFrame->UserData, ReadyFrame->GetTextureResource());
 				}
 				else
 				{
 					TRACE_CPUPROFILER_EVENT_SCOPE(UMediaCapture::RHIResourceCaptured);
-					MediaCapture->OnRHIResourceCaptured_RenderingThread(ReadyFrame->CaptureBaseData, ReadyFrame->UserData, ReadyFrame->GetBufferResource());
+					MediaCapture->OnRHIResourceCaptured_RenderingThread(RHICmdList, ReadyFrame->CaptureBaseData, ReadyFrame->UserData, ReadyFrame->GetBufferResource());
 				}
 			}
 			else
@@ -273,7 +275,7 @@ void FMediaCaptureHelper::OnReadbackComplete(FRHICommandList& RHICmdList, UMedia
 
 				// Readback should be ready since we're after the sync point.
 				SCOPE_CYCLE_COUNTER(STAT_MediaCaptureHelper_RenderThread_LockResource);
-				ColorDataBuffer = ReadyFrame->Lock(FRHICommandListExecutor::GetImmediateCommandList(), RowStride);
+				ColorDataBuffer = ReadyFrame->Lock(RHICmdList, RowStride);
 
 				if (ensure(ColorDataBuffer))
 				{
diff --git a/Engine/Plugins/Media/MediaIOFramework/Source/MediaIOCore/Private/MediaCaptureHelper.h b/Engine/Plugins/Media/MediaIOFramework/Source/MediaIOCore/Private/MediaCaptureHelper.h
--- a/Engine/Plugins/Media/MediaIOFramework/Source/MediaIOCore/Private/MediaCaptureHelper.h
+++ b/Engine/Plugins/Media/MediaIOFramework/Source/MediaIOCore/Private/MediaCaptureHelper.h
@@ -61,6 +61,6 @@ namespace UE::MediaCaptureData
 		 * Function called when the last copy of the media capture pipeline is completed. 
 		 * This will handle calling the right capture callback on the MediaCapture (ie. OnRHIResourceCaptured_RenderingThread)
 		 */
-		static void OnReadbackComplete(FRHICommandList& RHICmdList, UMediaCapture* MediaCapture, TSharedPtr<UE::MediaCaptureData::FCaptureFrame> ReadyFrame);
+		static void OnReadbackComplete(UMediaCapture* MediaCapture, TSharedPtr<UE::MediaCaptureData::FCaptureFrame> ReadyFrame);
 	};
 }
diff --git a/Engine/Plugins/Media/MediaIOFramework/Source/MediaIOCore/Private/MediaCaptureSyncPointWatcher.cpp b/Engine/Plugins/Media/MediaIOFramework/Source/MediaIOCore/Private/MediaCaptureSyncPointWatcher.cpp
--- a/Engine/Plugins/Media/MediaIOFramework/Source/MediaIOCore/Private/MediaCaptureSyncPointWatcher.cpp
+++ b/Engine/Plugins/Media/MediaIOFramework/Source/MediaIOCore/Private/MediaCaptureSyncPointWatcher.cpp
@@ -87,15 +87,15 @@ void FSyncPointWatcher::ProcessPendingCapture(const FPendingCaptureData& Data)
 
 		if (MediaCapture->UseAnyThreadCapture())
 		{
-			FMediaCaptureHelper::OnReadbackComplete(FRHICommandListExecutor::GetImmediateCommandList(), MediaCapture, Data.CapturingFrame);
+			FMediaCaptureHelper::OnReadbackComplete(MediaCapture, Data.CapturingFrame);
 		}
 		else
 		{
 			++MediaCapture->WaitingForRenderCommandExecutionCounter;
 
-			ENQUEUE_RENDER_COMMAND(MediaOutputCaptureReadbackComplete)([OwnerCapture = MediaCapture, CapturingFrame = Data.CapturingFrame](FRHICommandList& RHICommandList)
+			ENQUEUE_RENDER_COMMAND(MediaOutputCaptureReadbackComplete)([OwnerCapture = MediaCapture, CapturingFrame = Data.CapturingFrame](FRHICommandListImmediate&)
 			{
-				FMediaCaptureHelper::OnReadbackComplete(RHICommandList, OwnerCapture, CapturingFrame);
+				FMediaCaptureHelper::OnReadbackComplete(OwnerCapture, CapturingFrame);
 			});
 		}
 	}
diff --git a/Engine/Plugins/Media/MediaIOFramework/Source/MediaIOCore/Private/MediaIOCoreTextureSampleConverter.cpp b/Engine/Plugins/Media/MediaIOFramework/Source/MediaIOCore/Private/MediaIOCoreTextureSampleConverter.cpp
--- a/Engine/Plugins/Media/MediaIOFramework/Source/MediaIOCore/Private/MediaIOCoreTextureSampleConverter.cpp
+++ b/Engine/Plugins/Media/MediaIOFramework/Source/MediaIOCore/Private/MediaIOCoreTextureSampleConverter.cpp
@@ -12,7 +12,7 @@ void FMediaIOCoreTextureSampleConverter::Setup(const TSharedPtr<FMediaIOCoreText
 	JITRProxySample = InSample;
 }
 
-bool FMediaIOCoreTextureSampleConverter::Convert(FTextureRHIRef& InDestinationTexture, const FConversionHints& Hints)
+bool FMediaIOCoreTextureSampleConverter::Convert(FRHICommandListImmediate& RHICmdList, FTextureRHIRef& InDestinationTexture, const FConversionHints& Hints)
 {
 	TRACE_CPUPROFILER_EVENT_SCOPE(JITR);
 
@@ -31,7 +31,7 @@ bool FMediaIOCoreTextureSampleConverter::Convert(FTextureRHIRef& InDestinationTe
 	}
 
 	// Let the player pick a proper sample and render it into this proxy's texture (JITR)
-	return PlayerPtr->JustInTimeSampleRender_RenderThread(SamplePtr);
+	return PlayerPtr->JustInTimeSampleRender_RenderThread(RHICmdList, SamplePtr);
 }
 
 uint32 FMediaIOCoreTextureSampleConverter::GetConverterInfoFlags() const
diff --git a/Engine/Plugins/Media/MediaIOFramework/Source/MediaIOCore/Private/Player/MediaIOCorePlayerBase.cpp b/Engine/Plugins/Media/MediaIOFramework/Source/MediaIOCore/Private/Player/MediaIOCorePlayerBase.cpp
--- a/Engine/Plugins/Media/MediaIOFramework/Source/MediaIOCore/Private/Player/MediaIOCorePlayerBase.cpp
+++ b/Engine/Plugins/Media/MediaIOFramework/Source/MediaIOCore/Private/Player/MediaIOCorePlayerBase.cpp
@@ -647,39 +647,39 @@ void FMediaIOCorePlayerBase::RegisterSampleBuffer(const TSharedPtr<FMediaIOCoreT
 		RegisteredBuffers.Add(InSample->GetMutableBuffer());
 
 		// Enqueue buffer registration on the render  thread to ensure it happens before the call to TransferTexture
-		ENQUEUE_RENDER_COMMAND(RegisterSampleBuffers)([this, InSample](FRHICommandList& CommandList)
+		ENQUEUE_RENDER_COMMAND(RegisterSampleBuffers)([this, InSample](FRHICommandListImmediate&)
+		{
+			if (!InSample->GetMutableBuffer())
 			{
-				if (!InSample->GetMutableBuffer())
-				{
-					UE_LOG(LogMediaIOCore, Error, TEXT("A buffer was not available while performing a gpu texture transfer."));
-					return;
-				}
+				UE_LOG(LogMediaIOCore, Error, TEXT("A buffer was not available while performing a gpu texture transfer."));
+				return;
+			}
 
-				if (!InSample->GetTexture())
-				{
-					UE_LOG(LogMediaIOCore, Error, TEXT("A texture was not available while performing a gpu texture transfer."));
-					return;
-				}
+			if (!InSample->GetTexture())
+			{
+				UE_LOG(LogMediaIOCore, Error, TEXT("A texture was not available while performing a gpu texture transfer."));
+				return;
+			}
 
-				const uint32 TextureWidth = InSample->GetTexture()->GetDesc().Extent.X;
-				const uint32 TextureHeight = InSample->GetTexture()->GetDesc().Extent.Y;
-				uint32 TextureStride = TextureWidth * 4;
+			const uint32 TextureWidth = InSample->GetTexture()->GetDesc().Extent.X;
+			const uint32 TextureHeight = InSample->GetTexture()->GetDesc().Extent.Y;
+			uint32 TextureStride = TextureWidth * 4;
 
-				EPixelFormat Format = InSample->GetTexture()->GetFormat();
-				if (Format == PF_R32G32B32A32_UINT)
-				{
-					TextureStride *= 4;
-				}
+			EPixelFormat Format = InSample->GetTexture()->GetFormat();
+			if (Format == PF_R32G32B32A32_UINT)
+			{
+				TextureStride *= 4;
+			}
 
-				UE_LOG(LogMediaIOCore, Verbose, TEXT("Registering buffer %u"), reinterpret_cast<uintptr_t>(InSample->GetMutableBuffer()));
-				UE::GPUTextureTransfer::FRegisterDMABufferArgs Args;
-				Args.Buffer = InSample->GetMutableBuffer();
-				Args.Width = TextureWidth;
-				Args.Height = TextureHeight;
-				Args.Stride = TextureStride;
-				Args.PixelFormat = Format == PF_B8G8R8A8 ? UE::GPUTextureTransfer::EPixelFormat::PF_8Bit : UE::GPUTextureTransfer::EPixelFormat::PF_10Bit;
-				GPUTextureTransfer->RegisterBuffer(Args);
-			});
+			UE_LOG(LogMediaIOCore, Verbose, TEXT("Registering buffer %u"), reinterpret_cast<uintptr_t>(InSample->GetMutableBuffer()));
+			UE::GPUTextureTransfer::FRegisterDMABufferArgs Args;
+			Args.Buffer = InSample->GetMutableBuffer();
+			Args.Width = TextureWidth;
+			Args.Height = TextureHeight;
+			Args.Stride = TextureStride;
+			Args.PixelFormat = Format == PF_B8G8R8A8 ? UE::GPUTextureTransfer::EPixelFormat::PF_8Bit : UE::GPUTextureTransfer::EPixelFormat::PF_10Bit;
+			GPUTextureTransfer->RegisterBuffer(Args);
+		});
 	}
 	else
 	{
@@ -689,16 +689,16 @@ void FMediaIOCorePlayerBase::RegisterSampleBuffer(const TSharedPtr<FMediaIOCoreT
 
 void FMediaIOCorePlayerBase::UnregisterSampleBuffers()
 {
-	ENQUEUE_RENDER_COMMAND(UnregisterSampleBuffers)([TextureTransfer = GPUTextureTransfer, BuffersToUnregister = RegisteredBuffers](FRHICommandList& CommandList)
+	ENQUEUE_RENDER_COMMAND(UnregisterSampleBuffers)([TextureTransfer = GPUTextureTransfer, BuffersToUnregister = RegisteredBuffers](FRHICommandListImmediate&)
+	{
+		if (TextureTransfer)
 		{
-			if (TextureTransfer)
+			for (void* RegisteredBuffer : BuffersToUnregister)
 			{
-				for (void* RegisteredBuffer : BuffersToUnregister)
-				{
-					TextureTransfer->UnregisterBuffer(RegisteredBuffer);
-				}
+				TextureTransfer->UnregisterBuffer(RegisteredBuffer);
 			}
-		});
+		}
+	});
 
 	RegisteredBuffers.Reset();
 }
@@ -852,12 +852,12 @@ void FMediaIOCorePlayerBase::PreGPUTransfer(const TSharedPtr<FMediaIOCoreTexture
 	{
 		InSample->SetTexture(Texture);
 		InSample->SetDestructionCallback([MediaPlayerWeakPtr = TWeakPtr<FMediaIOCorePlayerBase>(AsShared())](TRefCountPtr<FRHITexture> InTexture)
+		{
+			if (TSharedPtr<FMediaIOCorePlayerBase> MediaPlayerPtr = MediaPlayerWeakPtr.Pin())
 			{
-				if (TSharedPtr<FMediaIOCorePlayerBase> MediaPlayerPtr = MediaPlayerWeakPtr.Pin())
-				{
-					MediaPlayerPtr->OnSampleDestroyed(InTexture);
-				}
-			});
+				MediaPlayerPtr->OnSampleDestroyed(InTexture);
+			}
+		});
 
 		RegisterSampleBuffer(InSample);
 	}
@@ -1132,7 +1132,7 @@ void FMediaIOCorePlayerBase::AddVideoSample(const TSharedRef<FMediaIOCoreTexture
 	}
 }
 
-bool FMediaIOCorePlayerBase::JustInTimeSampleRender_RenderThread(TSharedPtr<FMediaIOCoreTextureSampleBase>& JITRProxySample)
+bool FMediaIOCorePlayerBase::JustInTimeSampleRender_RenderThread(FRHICommandListImmediate& RHICmdList, TSharedPtr<FMediaIOCoreTextureSampleBase>& JITRProxySample)
 {
 	checkSlow(IsInRenderingThread());
 
@@ -1163,7 +1163,7 @@ bool FMediaIOCorePlayerBase::JustInTimeSampleRender_RenderThread(TSharedPtr<FMed
 	JITRProxySample->CopyConfiguration(SourceSample);
 
 	// Now we know which sample to use, transfer its texture
-	TransferTexture_RenderThread(SourceSample, JITRProxySample);
+	TransferTexture_RenderThread(RHICmdList, SourceSample, JITRProxySample);
 
 	// Update frame number so we won't render this sample again
 	LastEngineRTFrameThatUpdatedJustInTime = GFrameCounterRenderThread;
@@ -1310,10 +1310,8 @@ TSharedPtr<FMediaIOCoreTextureSampleBase> FMediaIOCorePlayerBase::PickSampleToRe
 	return PickSampleToRenderForTimeSynchronized_RenderThread(JITRProxySample);
 }
 
-void FMediaIOCorePlayerBase::TransferTexture_RenderThread(const TSharedPtr<FMediaIOCoreTextureSampleBase>& Sample, const TSharedPtr<FMediaIOCoreTextureSampleBase>& JITRProxySample)
+void FMediaIOCorePlayerBase::TransferTexture_RenderThread(FRHICommandListImmediate& RHICmdList, const TSharedPtr<FMediaIOCoreTextureSampleBase>& Sample, const TSharedPtr<FMediaIOCoreTextureSampleBase>& JITRProxySample)
 {
-	FRHICommandListImmediate& RHICmdList = FRHICommandListImmediate::Get();
-
 	// DMA P2P transfer
 	const bool bIsSampleAwaitingForGpuTransfer = Sample->IsAwaitingForGPUTransfer();
 	if (CanUseGPUTextureTransfer() && bIsSampleAwaitingForGpuTransfer)
diff --git a/Engine/Plugins/Media/MediaIOFramework/Source/MediaIOCore/Private/Shared/MediaIOCoreTextureSampleBase.cpp b/Engine/Plugins/Media/MediaIOFramework/Source/MediaIOCore/Private/Shared/MediaIOCoreTextureSampleBase.cpp
--- a/Engine/Plugins/Media/MediaIOFramework/Source/MediaIOCore/Private/Shared/MediaIOCoreTextureSampleBase.cpp
+++ b/Engine/Plugins/Media/MediaIOFramework/Source/MediaIOCore/Private/Shared/MediaIOCoreTextureSampleBase.cpp
@@ -207,10 +207,8 @@ void FMediaIOCoreTextureSampleBase::CopyConfiguration(const TSharedPtr<FMediaIOC
 	OriginalSample = SourceSample;
 }
 
-bool FMediaIOCoreTextureSampleBase::ApplyColorConversion(FTextureRHIRef& InSrcTexture, FTextureRHIRef& InDstTexture)
+bool FMediaIOCoreTextureSampleBase::ApplyColorConversion(FRHICommandListImmediate& RHICmdList, FTextureRHIRef& InSrcTexture, FTextureRHIRef& InDstTexture)
 {
-	FRHICommandListImmediate& RHICmdList = FRHICommandListExecutor::GetImmediateCommandList();
-
 	if (!CachedOCIOResources.IsValid())
 	{
 		if (ColorConversionSettings.IsValid() && ColorConversionSettings->IsValid())
diff --git a/Engine/Plugins/Media/MediaIOFramework/Source/MediaIOCore/Public/MediaCapture.h b/Engine/Plugins/Media/MediaIOFramework/Source/MediaIOCore/Public/MediaCapture.h
--- a/Engine/Plugins/Media/MediaIOFramework/Source/MediaIOCore/Public/MediaCapture.h
+++ b/Engine/Plugins/Media/MediaIOFramework/Source/MediaIOCore/Public/MediaCapture.h
@@ -520,8 +520,8 @@ protected:
 	 * The callback in called from a critical point. If you intend to process the texture/buffer, do so in another thread.
 	 * The texture is valid for the duration of the callback.
 	 */
-	virtual void OnRHIResourceCaptured_RenderingThread(const FCaptureBaseData& InBaseData, TSharedPtr<FMediaCaptureUserData, ESPMode::ThreadSafe> InUserData, FTextureRHIRef InTexture) { }
-	virtual void OnRHIResourceCaptured_RenderingThread(const FCaptureBaseData& InBaseData, TSharedPtr<FMediaCaptureUserData, ESPMode::ThreadSafe> InUserData, FBufferRHIRef InBuffer) { }
+	virtual void OnRHIResourceCaptured_RenderingThread(FRHICommandListImmediate& RHICmdList, const FCaptureBaseData& InBaseData, TSharedPtr<FMediaCaptureUserData, ESPMode::ThreadSafe> InUserData, FTextureRHIRef InTexture) { }
+	virtual void OnRHIResourceCaptured_RenderingThread(FRHICommandListImmediate& RHICmdList, const FCaptureBaseData& InBaseData, TSharedPtr<FMediaCaptureUserData, ESPMode::ThreadSafe> InUserData, FBufferRHIRef InBuffer) { }
 	
 	/**
 	 * AnyThread version of the above callbacks.
diff --git a/Engine/Plugins/Media/MediaIOFramework/Source/MediaIOCore/Public/MediaIOCorePlayerBase.h b/Engine/Plugins/Media/MediaIOFramework/Source/MediaIOCore/Public/MediaIOCorePlayerBase.h
--- a/Engine/Plugins/Media/MediaIOFramework/Source/MediaIOCore/Public/MediaIOCorePlayerBase.h
+++ b/Engine/Plugins/Media/MediaIOFramework/Source/MediaIOCore/Public/MediaIOCorePlayerBase.h
@@ -167,7 +167,7 @@ public:
 	 * Just in time sample rendering. This method is responsible for late sample picking,
 	 * then rendering it into the proxy sample provided.
 	 */
-	bool JustInTimeSampleRender_RenderThread(TSharedPtr<FMediaIOCoreTextureSampleBase>& JITRProxySample);
+	bool JustInTimeSampleRender_RenderThread(FRHICommandListImmediate& RHICmdList, TSharedPtr<FMediaIOCoreTextureSampleBase>& JITRProxySample);
 
 protected:
 
@@ -263,7 +263,7 @@ protected:
 	 * A wrapper method responsible for transferring of the sample textures into GPU memory based
 	 * on the current settings and hardware capabilities.
 	 */
-	void TransferTexture_RenderThread(const TSharedPtr<FMediaIOCoreTextureSampleBase>& Sample, const TSharedPtr<FMediaIOCoreTextureSampleBase>& JITRProxySample);
+	void TransferTexture_RenderThread(FRHICommandListImmediate& RHICmdList, const TSharedPtr<FMediaIOCoreTextureSampleBase>& Sample, const TSharedPtr<FMediaIOCoreTextureSampleBase>& JITRProxySample);
 
 protected:
 	/** Url used to open the media player. */
diff --git a/Engine/Plugins/Media/MediaIOFramework/Source/MediaIOCore/Public/MediaIOCoreTextureSampleBase.h b/Engine/Plugins/Media/MediaIOFramework/Source/MediaIOCore/Public/MediaIOCoreTextureSampleBase.h
--- a/Engine/Plugins/Media/MediaIOFramework/Source/MediaIOCore/Public/MediaIOCoreTextureSampleBase.h
+++ b/Engine/Plugins/Media/MediaIOFramework/Source/MediaIOCore/Public/MediaIOCoreTextureSampleBase.h
@@ -340,7 +340,7 @@ public:
 	}
 
 	//~ IMediaTextureSampleColorConverter interface
-	virtual bool ApplyColorConversion(FTextureRHIRef& InSrcTexture, FTextureRHIRef& InDstTexture) override;
+	virtual bool ApplyColorConversion(FRHICommandListImmediate& RHICmdList, FTextureRHIRef& InSrcTexture, FTextureRHIRef& InDstTexture) override;
 
 	void* GetMutableBuffer()
 	{
diff --git a/Engine/Plugins/Media/MediaIOFramework/Source/MediaIOCore/Public/MediaIOCoreTextureSampleConverter.h b/Engine/Plugins/Media/MediaIOFramework/Source/MediaIOCore/Public/MediaIOCoreTextureSampleConverter.h
--- a/Engine/Plugins/Media/MediaIOFramework/Source/MediaIOCore/Public/MediaIOCoreTextureSampleConverter.h
+++ b/Engine/Plugins/Media/MediaIOFramework/Source/MediaIOCore/Public/MediaIOCoreTextureSampleConverter.h
@@ -29,7 +29,7 @@ public:
 
 public:
 	//~ Begin IMediaTextureSampleConverter interface
-	virtual bool Convert(FTextureRHIRef& InDstTexture, const FConversionHints& Hints) override;
+	virtual bool Convert(FRHICommandListImmediate& RHICmdList, FTextureRHIRef& InDstTexture, const FConversionHints& Hints) override;
 	virtual uint32 GetConverterInfoFlags() const override;
 	//~ End IMediaTextureSampleConverter interface
 
diff --git a/Engine/Plugins/Media/PixelStreaming/Source/PixelStreaming/Internal/PixelStreamingMediaIOCapture.h b/Engine/Plugins/Media/PixelStreaming/Source/PixelStreaming/Internal/PixelStreamingMediaIOCapture.h
--- a/Engine/Plugins/Media/PixelStreaming/Source/PixelStreaming/Internal/PixelStreamingMediaIOCapture.h
+++ b/Engine/Plugins/Media/PixelStreaming/Source/PixelStreaming/Internal/PixelStreamingMediaIOCapture.h
@@ -19,6 +19,7 @@ public:
 	* GPU copy methods
 	*/
 	virtual void OnRHIResourceCaptured_RenderingThread(
+		FRHICommandListImmediate& RHICmdList,
 		const FCaptureBaseData& InBaseData,
 		TSharedPtr<FMediaCaptureUserData, ESPMode::ThreadSafe> InUserData,
 		FTextureRHIRef InTexture) override;
diff --git a/Engine/Plugins/Media/PixelStreaming/Source/PixelStreaming/Private/PixelStreamingMediaIOCapture.cpp b/Engine/Plugins/Media/PixelStreaming/Source/PixelStreaming/Private/PixelStreamingMediaIOCapture.cpp
--- a/Engine/Plugins/Media/PixelStreaming/Source/PixelStreaming/Private/PixelStreamingMediaIOCapture.cpp
+++ b/Engine/Plugins/Media/PixelStreaming/Source/PixelStreaming/Private/PixelStreamingMediaIOCapture.cpp
@@ -15,6 +15,7 @@
 #include "RenderGraphUtils.h"
 
 void UPixelStreamingMediaIOCapture::OnRHIResourceCaptured_RenderingThread(
+	FRHICommandListImmediate& /*RHICmdList*/,
 	const FCaptureBaseData& InBaseData,
 	TSharedPtr<FMediaCaptureUserData, ESPMode::ThreadSafe> InUserData,
 	FTextureRHIRef InTexture)
diff --git a/Engine/Plugins/Media/PixelStreaming/Source/PixelStreaming/Private/PixelStreamingVideoSink.cpp b/Engine/Plugins/Media/PixelStreaming/Source/PixelStreaming/Private/PixelStreamingVideoSink.cpp
--- a/Engine/Plugins/Media/PixelStreaming/Source/PixelStreaming/Private/PixelStreamingVideoSink.cpp
+++ b/Engine/Plugins/Media/PixelStreaming/Source/PixelStreaming/Private/PixelStreamingVideoSink.cpp
@@ -52,7 +52,7 @@ void FPixelStreamingVideoSink::OnFrame(const webrtc::VideoFrame& Frame)
     	    FScopeLock Lock(&RenderSyncContext);
 
     	    const FIntPoint FrameSize = FIntPoint(SizeX, SizeY);
-    	    FRHICommandListImmediate& RHICmdList = FRHICommandListExecutor::GetImmediateCommandList();
+    	    FRHICommandListImmediate& RHICmdList = FRHICommandListImmediate::Get();
 
             FTextureRHIRef SourceTexture;
 
@@ -78,7 +78,7 @@ void FPixelStreamingVideoSink::OnFrame(const webrtc::VideoFrame& Frame)
     	                .SetInitialState(ERHIAccess::SRVMask);
 #endif
 
-            	SourceTexture = RHICreateTexture(RenderTargetTextureDesc);
+            	SourceTexture = RHICmdList.CreateTexture(RenderTargetTextureDesc);
 
             	// Find a free target-able texture from the render pool
             	GRenderTargetPool.FindFreeElement(RHICmdList,
@@ -90,7 +90,7 @@ void FPixelStreamingVideoSink::OnFrame(const webrtc::VideoFrame& Frame)
         		const FUpdateTextureRegion2D Region(0, 0, 0, 0, FrameSize.X, FrameSize.Y);
 
         		// Set the Pixel data of the webrtc Frame to the SourceTexture
-        		RHIUpdateTexture2D(SourceTexture, 0, Region, FrameSize.X * 4, Buffer.GetData());
+        		RHICmdList.UpdateTexture2D(SourceTexture, 0, Region, FrameSize.X * 4, Buffer.GetData());
 
             	OnFrame(SourceTexture);
         	}
diff --git a/Engine/Plugins/Media/PixelStreaming/Source/PixelStreamingEditor/Private/PixelStreamingVideoInputBackBufferComposited.cpp b/Engine/Plugins/Media/PixelStreaming/Source/PixelStreamingEditor/Private/PixelStreamingVideoInputBackBufferComposited.cpp
--- a/Engine/Plugins/Media/PixelStreaming/Source/PixelStreamingEditor/Private/PixelStreamingVideoInputBackBufferComposited.cpp
+++ b/Engine/Plugins/Media/PixelStreaming/Source/PixelStreamingEditor/Private/PixelStreamingVideoInputBackBufferComposited.cpp
@@ -139,8 +139,7 @@ void FPixelStreamingVideoInputBackBufferComposited::OnBackBufferReady(SWindow& S
 		}
 
 		{
-			FRHICommandListImmediate& RHICmdList = FRHICommandListExecutor::GetImmediateCommandList();
-			FRDGBuilder GraphBuilder(RHICmdList);
+			FRDGBuilder GraphBuilder(FRHICommandListImmediate::Get());
 
 			// Register an external RDG texture from the provided frame buffer
 			FRDGTextureRef InputTexture = GraphBuilder.RegisterExternalTexture(CreateRenderTarget(FrameBuffer, *SlateWindow.GetTitle().ToString()));
@@ -182,8 +181,6 @@ void FPixelStreamingVideoInputBackBufferComposited::OnBackBufferReady(SWindow& S
 
 void FPixelStreamingVideoInputBackBufferComposited::CompositeWindows()
 {
-	FRHICommandListImmediate& RHICmdList = FRHICommandListExecutor::GetImmediateCommandList();
-
 	// Process all of the windows we will need to render. This processing step finds the extents of the
 	// composited texture as well as the top-left point
 	FIntPoint TopLeft = FIntPoint(MAX_int32, MAX_int32);
@@ -208,9 +205,10 @@ void FPixelStreamingVideoInputBackBufferComposited::CompositeWindows()
 	PermutationVector.Set<FModifyAlphaSwizzleRgbaPS::FConversionOp>(ConversionOperation);
 
 	FTextureRHIRef OutTexture = nullptr;
-	{	// FRDGBuilder uses a global allocator which can cause race conditions
+	{
+		// FRDGBuilder uses a global allocator which can cause race conditions
 		// To prevent issues its lifetime needs to end as soon as it has executed
-		FRDGBuilder GraphBuilder(RHICmdList);
+		FRDGBuilder GraphBuilder(FRHICommandListImmediate::Get());
 
 		// Clamp the texture dimensions to ensure no RHI crashes
 		// Create an RDG texture that is the size of our extent for use as the composited frame
@@ -251,7 +249,7 @@ void FPixelStreamingVideoInputBackBufferComposited::CompositeWindows()
 			ViewInitOptions.ViewRotationMatrix = FMatrix::Identity;
 			ViewInitOptions.ProjectionMatrix = FMatrix::Identity;
 
-			GetRendererModule().CreateAndInitSingleView(RHICmdList, &ViewFamily, &ViewInitOptions);
+			GetRendererModule().CreateAndInitSingleView(GraphBuilder.RHICmdList, &ViewFamily, &ViewInitOptions);
 			const FSceneView& View = *ViewFamily.Views[0];
 
 			TShaderMapRef<FModifyAlphaSwizzleRgbaPS> PixelShader(GlobalShaderMap, PermutationVector);
@@ -267,7 +265,7 @@ void FPixelStreamingVideoInputBackBufferComposited::CompositeWindows()
 			RDG_EVENT_NAME("PushCompositedFrame"),
 			PassParameters,
 			ERDGPassFlags::Readback,
-			[CompositedTexture, &OutTexture, this](FRHICommandList& RHICmdList) {
+			[CompositedTexture, &OutTexture, this](FRHICommandList&) {
 				// Our composition is complete out the underlying rhi resource
 				OutTexture = CompositedTexture->GetRHI();
 			});
diff --git a/Engine/Plugins/Media/WmfMedia/Source/WmfMedia/Private/Wmf/WmfMediaHardwareVideoDecodingRendering.cpp b/Engine/Plugins/Media/WmfMedia/Source/WmfMedia/Private/Wmf/WmfMediaHardwareVideoDecodingRendering.cpp
--- a/Engine/Plugins/Media/WmfMedia/Source/WmfMedia/Private/Wmf/WmfMediaHardwareVideoDecodingRendering.cpp
+++ b/Engine/Plugins/Media/WmfMedia/Source/WmfMedia/Private/Wmf/WmfMediaHardwareVideoDecodingRendering.cpp
@@ -104,7 +104,7 @@ struct FRHICommandCopyResource final : public FRHICommand<FRHICommandCopyResourc
 	}
 };
 
-bool FWmfMediaHardwareVideoDecodingParameters::ConvertTextureFormat_RenderThread(FWmfMediaHardwareVideoDecodingTextureSample* InSample, FTextureRHIRef InDstTexture)
+bool FWmfMediaHardwareVideoDecodingParameters::ConvertTextureFormat_RenderThread(FRHICommandListImmediate& RHICmdList, FWmfMediaHardwareVideoDecodingTextureSample* InSample, FTextureRHIRef InDstTexture)
 {
 	LLM_SCOPE(ELLMTag::MediaStreaming);
 
@@ -113,7 +113,6 @@ bool FWmfMediaHardwareVideoDecodingParameters::ConvertTextureFormat_RenderThread
 		return false;
 	}
 
-	FRHICommandListImmediate& RHICmdList = FRHICommandListImmediate::Get();
 	check(InSample);
 	
 	TComPtr<ID3D11Texture2D> SampleTexture = InSample->GetSourceTexture();
@@ -148,7 +147,7 @@ bool FWmfMediaHardwareVideoDecodingParameters::ConvertTextureFormat_RenderThread
 			uint32 Stride = InSample->GetStride();
 			uint32 Height = InSample->GetDim().Y;
 			FUpdateTextureRegion2D Region(0, 0, 0, 0, InSample->GetDim().X, Height);
-			RHIUpdateTexture2D(SampleDestinationTexture, 0, Region, Stride, Data);
+			RHICmdList.UpdateTexture2D(SampleDestinationTexture, 0, Region, Stride, Data);
 		}
 		else
 		{
@@ -251,7 +250,7 @@ bool FWmfMediaHardwareVideoDecodingParameters::ConvertTextureFormat_RenderThread
 				uint32 Stride = InSample->GetDim().X * 2;
 				uint32 Height = InSample->GetDim().Y;
 				FUpdateTextureRegion2D Region(0, 0, 0, 0, InSample->GetDim().X, Height);
-				RHIUpdateTexture2D(SampleDestinationAlphaTexture, 0, Region, Stride, Data);
+				RHICmdList.UpdateTexture2D(SampleDestinationAlphaTexture, 0, Region, Stride, Data);
 			}
 			
 			TShaderMapRef< FHardwareVideoDecodingVS > VertexShader(GlobalShaderMap);
diff --git a/Engine/Plugins/Media/WmfMedia/Source/WmfMedia/Public/WmfMediaHardwareVideoDecodingRendering.h b/Engine/Plugins/Media/WmfMedia/Source/WmfMedia/Public/WmfMediaHardwareVideoDecodingRendering.h
--- a/Engine/Plugins/Media/WmfMedia/Source/WmfMedia/Public/WmfMediaHardwareVideoDecodingRendering.h
+++ b/Engine/Plugins/Media/WmfMedia/Source/WmfMedia/Public/WmfMediaHardwareVideoDecodingRendering.h
@@ -9,5 +9,5 @@ class FWmfMediaHardwareVideoDecodingTextureSample;
 class FWmfMediaHardwareVideoDecodingParameters
 {
 public:
-	WMFMEDIA_API static bool ConvertTextureFormat_RenderThread(FWmfMediaHardwareVideoDecodingTextureSample* InSample, FTextureRHIRef InDstTexture);
+	WMFMEDIA_API static bool ConvertTextureFormat_RenderThread(FRHICommandListImmediate& RHICmdList, FWmfMediaHardwareVideoDecodingTextureSample* InSample, FTextureRHIRef InDstTexture);
 };
diff --git a/Engine/Plugins/Media/WmfMedia/Source/WmfMedia/Public/WmfMediaHardwareVideoDecodingTextureSample.h b/Engine/Plugins/Media/WmfMedia/Source/WmfMedia/Public/WmfMediaHardwareVideoDecodingTextureSample.h
--- a/Engine/Plugins/Media/WmfMedia/Source/WmfMedia/Public/WmfMediaHardwareVideoDecodingTextureSample.h
+++ b/Engine/Plugins/Media/WmfMedia/Source/WmfMedia/Public/WmfMediaHardwareVideoDecodingTextureSample.h
@@ -94,9 +94,9 @@ public:
 	/**
 	 * Texture sample convert using hardware video decoding.
 	 */
-	virtual bool Convert(FTextureRHIRef & InDstTexture, const FConversionHints & Hints) override
+	virtual bool Convert(FRHICommandListImmediate& RHICmdList, FTextureRHIRef& InDstTexture, const FConversionHints& Hints) override
 	{
-		FWmfMediaHardwareVideoDecodingParameters::ConvertTextureFormat_RenderThread(this, InDstTexture);
+		FWmfMediaHardwareVideoDecodingParameters::ConvertTextureFormat_RenderThread(RHICmdList, this, InDstTexture);
 		return true;
 	}
 
diff --git a/Engine/Plugins/Runtime/WindowsMoviePlayer/Source/WindowsMoviePlayer/Private/WindowsMovieStreamer.cpp b/Engine/Plugins/Runtime/WindowsMoviePlayer/Source/WindowsMoviePlayer/Private/WindowsMovieStreamer.cpp
--- a/Engine/Plugins/Runtime/WindowsMoviePlayer/Source/WindowsMoviePlayer/Private/WindowsMovieStreamer.cpp
+++ b/Engine/Plugins/Runtime/WindowsMoviePlayer/Source/WindowsMoviePlayer/Private/WindowsMovieStreamer.cpp
@@ -73,7 +73,7 @@ void FMediaFoundationMovieStreamer::ForceCompletion()
 	CloseMovie();
 }
 
-void FMediaFoundationMovieStreamer::ConvertSample()
+void FMediaFoundationMovieStreamer::ConvertSample(FRHICommandListImmediate& RHICmdList)
 {
 	const bool SrgbOutput = false;
 	const bool bSampleIsOutputSrgb = false;
@@ -95,12 +95,12 @@ void FMediaFoundationMovieStreamer::ConvertSample()
 				.SetFlags(InputCreateFlags | ETextureCreateFlags::RenderTargetable | ETextureCreateFlags::ShaderResource)
 				.SetInitialState(ERHIAccess::SRVMask);
 
-			InputTarget = RHICreateTexture(Desc);
+			InputTarget = RHICmdList.CreateTexture(Desc);
 		}
 
 		// copy sample data to input render target
 		FUpdateTextureRegion2D Region(0, 0, 0, 0, SourceFormat.BufferDim.X, SourceFormat.BufferDim.Y);
-		RHIUpdateTexture2D(InputTarget, 0, Region, SourceFormat.BufferStride, TextureData.GetData());
+		RHICmdList.UpdateTexture2D(InputTarget, 0, Region, SourceFormat.BufferStride, TextureData.GetData());
 	}
 
 	const FIntPoint OutputDim = SourceFormat.OutputDim;
@@ -109,17 +109,15 @@ void FMediaFoundationMovieStreamer::ConvertSample()
 	FRHITexture* RenderTarget = CurrentTexture->GetRHIRef();
 
 	// perform the conversion
-	FRHICommandListImmediate& CommandList = FRHICommandListExecutor::GetImmediateCommandList();
-
-	CommandList.Transition(FRHITransitionInfo(RenderTarget, ERHIAccess::Unknown, ERHIAccess::RTV));
+	RHICmdList.Transition(FRHITransitionInfo(RenderTarget, ERHIAccess::Unknown, ERHIAccess::RTV));
 
 	FRHIRenderPassInfo RPInfo(RenderTarget, ERenderTargetActions::Load_Store);
-	CommandList.BeginRenderPass(RPInfo, TEXT("WindowsMovieConvertSample"));
+	RHICmdList.BeginRenderPass(RPInfo, TEXT("WindowsMovieConvertSample"));
 	{
 		FGraphicsPipelineStateInitializer GraphicsPSOInit;
 
-		CommandList.ApplyCachedRenderTargets(GraphicsPSOInit);
-		CommandList.SetViewport(0, 0, 0.0f, OutputDim.X, OutputDim.Y, 1.0f);
+		RHICmdList.ApplyCachedRenderTargets(GraphicsPSOInit);
+		RHICmdList.SetViewport(0, 0, 0.0f, OutputDim.X, OutputDim.Y, 1.0f);
 
 		GraphicsPSOInit.DepthStencilState = TStaticDepthStencilState<false, CF_Always>::GetRHI();
 		GraphicsPSOInit.RasterizerState = TStaticRasterizerState<>::GetRHI();
@@ -139,8 +137,8 @@ void FMediaFoundationMovieStreamer::ConvertSample()
 		{
 			TShaderMapRef<FBMPConvertPS> ConvertShader(ShaderMap);
 			GraphicsPSOInit.BoundShaderState.PixelShaderRHI = ConvertShader.GetPixelShader();
-			SetGraphicsPipelineState(CommandList, GraphicsPSOInit, 0);
-			SetShaderParametersLegacyPS(CommandList, ConvertShader, InputTarget, OutputDim, bSampleIsOutputSrgb && !SrgbOutput);
+			SetGraphicsPipelineState(RHICmdList, GraphicsPSOInit, 0);
+			SetShaderParametersLegacyPS(RHICmdList, ConvertShader, InputTarget, OutputDim, bSampleIsOutputSrgb && !SrgbOutput);
 		}
 		break;
 
@@ -148,8 +146,8 @@ void FMediaFoundationMovieStreamer::ConvertSample()
 		{
 			TShaderMapRef<FYUY2ConvertPS> ConvertShader(ShaderMap);
 			GraphicsPSOInit.BoundShaderState.PixelShaderRHI = ConvertShader.GetPixelShader();
-			SetGraphicsPipelineState(CommandList, GraphicsPSOInit, 0);
-			SetShaderParametersLegacyPS(CommandList, ConvertShader, InputTarget, OutputDim, MediaShaders::YuvToRgbRec709Scaled, MediaShaders::YUVOffset8bits, bSampleIsOutputSrgb);
+			SetGraphicsPipelineState(RHICmdList, GraphicsPSOInit, 0);
+			SetShaderParametersLegacyPS(RHICmdList, ConvertShader, InputTarget, OutputDim, MediaShaders::YuvToRgbRec709Scaled, MediaShaders::YUVOffset8bits, bSampleIsOutputSrgb);
 		}
 		break;
 
@@ -159,23 +157,23 @@ void FMediaFoundationMovieStreamer::ConvertSample()
 
 		// draw full size quad into render target
 		FBufferRHIRef VertexBuffer = CreateTempMediaVertexBuffer();
-		CommandList.SetStreamSource(0, VertexBuffer, 0);
+		RHICmdList.SetStreamSource(0, VertexBuffer, 0);
 		// set viewport to RT size
-		CommandList.SetViewport(0, 0, 0.0f, OutputDim.X, OutputDim.Y, 1.0f);
-		CommandList.DrawPrimitive(0, 2, 1);
+		RHICmdList.SetViewport(0, 0, 0.0f, OutputDim.X, OutputDim.Y, 1.0f);
+		RHICmdList.DrawPrimitive(0, 2, 1);
 	}
-	CommandList.EndRenderPass();
-	CommandList.Transition(FRHITransitionInfo(RenderTarget, ERHIAccess::Unknown, ERHIAccess::SRVGraphics));
+	RHICmdList.EndRenderPass();
+	RHICmdList.Transition(FRHITransitionInfo(RenderTarget, ERHIAccess::Unknown, ERHIAccess::SRVGraphics));
 }
 
 bool FMediaFoundationMovieStreamer::Tick(float DeltaTime)
 {
 	FSlateTexture2DRHIRef* CurrentTexture = Texture.Get();
-	check(IsInRenderingThread());
+	FRHICommandListImmediate& RHICmdList = FRHICommandListImmediate::Get();
 
 	if (CurrentTexture && !CurrentTexture->IsInitialized())
 	{
-		CurrentTexture->InitResource(FRHICommandListImmediate::Get());
+		CurrentTexture->InitResource(RHICmdList);
 	}
 
 	if (CurrentTexture && SampleGrabberCallback->GetIsSampleReadyToUpdate())
@@ -190,7 +188,7 @@ bool FMediaFoundationMovieStreamer::Tick(float DeltaTime)
 		}
 		else
 		{
-			ConvertSample();
+			ConvertSample(RHICmdList);
 		}
 
 		if (MovieViewport->GetViewportRenderTargetTexture() == nullptr)
diff --git a/Engine/Plugins/Runtime/WindowsMoviePlayer/Source/WindowsMoviePlayer/Private/WindowsMovieStreamer.h b/Engine/Plugins/Runtime/WindowsMoviePlayer/Source/WindowsMoviePlayer/Private/WindowsMovieStreamer.h
--- a/Engine/Plugins/Runtime/WindowsMoviePlayer/Source/WindowsMoviePlayer/Private/WindowsMovieStreamer.h
+++ b/Engine/Plugins/Runtime/WindowsMoviePlayer/Source/WindowsMoviePlayer/Private/WindowsMovieStreamer.h
@@ -53,7 +53,7 @@ private:
 	/** Cleans up rendering resources once movies are done playing */
 	void CleanupRenderingResources();
 	/** Use shader to generate final image */
-	void ConvertSample();
+	void ConvertSample(FRHICommandListImmediate& RHICmdList);
 
 private:
 	/** A list of all the stored movie paths we have enqueued for playing */
diff --git a/Engine/Plugins/Runtime/nDisplay/Source/SharedMemoryMedia/Private/SharedMemoryMediaPlayer.cpp b/Engine/Plugins/Runtime/nDisplay/Source/SharedMemoryMedia/Private/SharedMemoryMediaPlayer.cpp
--- a/Engine/Plugins/Runtime/nDisplay/Source/SharedMemoryMedia/Private/SharedMemoryMediaPlayer.cpp
+++ b/Engine/Plugins/Runtime/nDisplay/Source/SharedMemoryMedia/Private/SharedMemoryMediaPlayer.cpp
@@ -707,10 +707,8 @@ bool FSharedMemoryMediaPlayer::DetermineNextSourceFrame(uint64 FrameNumber, uint
 	return false;
 }
 
-void FSharedMemoryMediaPlayer::JustInTimeSampleRender()
+void FSharedMemoryMediaPlayer::JustInTimeSampleRender(FRHICommandListImmediate& RHICmdList)
 {
-	FRHICommandListImmediate& RHICmdList = FRHICommandListExecutor::GetImmediateCommandList();
-
 	// We only allow this function to run once per frame.
 	if (LastFrameNumberThatUpdatedJustInTime == GFrameCounterRenderThread)
 	{
diff --git a/Engine/Plugins/Runtime/nDisplay/Source/SharedMemoryMedia/Private/SharedMemoryMediaPlayer.h b/Engine/Plugins/Runtime/nDisplay/Source/SharedMemoryMedia/Private/SharedMemoryMediaPlayer.h
--- a/Engine/Plugins/Runtime/nDisplay/Source/SharedMemoryMedia/Private/SharedMemoryMediaPlayer.h
+++ b/Engine/Plugins/Runtime/nDisplay/Source/SharedMemoryMedia/Private/SharedMemoryMediaPlayer.h
@@ -208,7 +208,7 @@ protected:
 protected:
 
 	/** Enqueues a last minute wait and copy of the expected texture data for the renderer to use */
-	virtual void JustInTimeSampleRender();
+	virtual void JustInTimeSampleRender(FRHICommandListImmediate& RHICmdList);
 
 	/** Offsets the frame number by the specified latency. Used to associate resources with frame numbers being evaluated */
 	uint32 InputTextureFrameNumberForFrameNumber(uint32 FrameNumber) const;
diff --git a/Engine/Plugins/Runtime/nDisplay/Source/SharedMemoryMedia/Private/SharedMemoryMediaTextureSampleConverter.h b/Engine/Plugins/Runtime/nDisplay/Source/SharedMemoryMedia/Private/SharedMemoryMediaTextureSampleConverter.h
--- a/Engine/Plugins/Runtime/nDisplay/Source/SharedMemoryMedia/Private/SharedMemoryMediaTextureSampleConverter.h
+++ b/Engine/Plugins/Runtime/nDisplay/Source/SharedMemoryMedia/Private/SharedMemoryMediaTextureSampleConverter.h
@@ -22,14 +22,14 @@ public:
 		return ConverterInfoFlags_PreprocessOnly;
 	}
 
-	virtual bool Convert(FTextureRHIRef& InDstTexture, const FConversionHints& Hints) override
+	virtual bool Convert(FRHICommandListImmediate& RHICmdList, FTextureRHIRef& InDstTexture, const FConversionHints& Hints) override
 	{
 		if (!Player)
 		{
 			return false;
 		}
 
-		Player->JustInTimeSampleRender();
+		Player->JustInTimeSampleRender(RHICmdList);
 
 		return true;
 	}
diff --git a/Engine/Plugins/VirtualProduction/Rivermax/RivermaxMedia/Source/RivermaxMedia/Private/RivermaxMediaCapture.cpp b/Engine/Plugins/VirtualProduction/Rivermax/RivermaxMedia/Source/RivermaxMedia/Private/RivermaxMediaCapture.cpp
--- a/Engine/Plugins/VirtualProduction/Rivermax/RivermaxMedia/Source/RivermaxMedia/Private/RivermaxMediaCapture.cpp
+++ b/Engine/Plugins/VirtualProduction/Rivermax/RivermaxMedia/Source/RivermaxMedia/Private/RivermaxMediaCapture.cpp
@@ -325,7 +325,7 @@ void URivermaxMediaCapture::OnFrameCaptured_RenderingThread(const FCaptureBaseDa
 	OnFrameCapturedInternal_AnyThread(InBaseData, InUserData, InBuffer, Width, Height, BytesPerRow);
 }
 
-void URivermaxMediaCapture::OnRHIResourceCaptured_RenderingThread(const FCaptureBaseData& InBaseData, TSharedPtr<FMediaCaptureUserData, ESPMode::ThreadSafe> InUserData, FBufferRHIRef InBuffer)
+void URivermaxMediaCapture::OnRHIResourceCaptured_RenderingThread(FRHICommandListImmediate& /*RHICmdList*/, const FCaptureBaseData& InBaseData, TSharedPtr<FMediaCaptureUserData, ESPMode::ThreadSafe> InUserData, FBufferRHIRef InBuffer)
 {
 	TRACE_CPUPROFILER_EVENT_SCOPE(URivermaxMediaCapture::OnRHIResourceCaptured_RenderingThread);
 	OnRHIResourceCapturedInternal_AnyThread(InBaseData, InUserData, InBuffer);
diff --git a/Engine/Plugins/VirtualProduction/Rivermax/RivermaxMedia/Source/RivermaxMedia/Private/RivermaxMediaPlayer.cpp b/Engine/Plugins/VirtualProduction/Rivermax/RivermaxMedia/Source/RivermaxMedia/Private/RivermaxMediaPlayer.cpp
--- a/Engine/Plugins/VirtualProduction/Rivermax/RivermaxMedia/Source/RivermaxMedia/Private/RivermaxMediaPlayer.cpp
+++ b/Engine/Plugins/VirtualProduction/Rivermax/RivermaxMedia/Source/RivermaxMedia/Private/RivermaxMediaPlayer.cpp
@@ -1051,20 +1051,18 @@ namespace UE::RivermaxMedia
 			OutConverterSetup.GetGPUBufferFunc = [SampleWrapper]() { return SampleWrapper->Sample->GetGPUBuffer(); };
 
 			// Setup requirements for sample to be ready to be rendered
-			FRHICommandListImmediate& RHICmdList = FRHICommandListExecutor::GetImmediateCommandList();
-			RHICmdList.EnqueueLambda(
-				[NextFrameExpectations, this](FRHICommandList& RHICmdList)
+			FRHICommandListImmediate& RHICmdList = FRHICommandListImmediate::Get();
+			RHICmdList.EnqueueLambda([NextFrameExpectations, this](FRHICommandList&)
+			{
+				FWaitConditionFunc ReadyToRenderConditionFunc = [](const TSharedPtr<FRivermaxSampleWrapper>& Sample)
 				{
-					FWaitConditionFunc ReadyToRenderConditionFunc = [](const TSharedPtr<FRivermaxSampleWrapper>& Sample)
-					{
-						return Sample->bIsReadyToRender.load();
-					};
+					return Sample->bIsReadyToRender.load();
+				};
 
-					TRACE_CPUPROFILER_EVENT_SCOPE(RmaxWaitSampleReadyness);
-					constexpr bool bCanTimeout = true;
-					WaitForSample(NextFrameExpectations, MoveTemp(ReadyToRenderConditionFunc), bCanTimeout);
-				}
-			);
+				TRACE_CPUPROFILER_EVENT_SCOPE(RmaxWaitSampleReadyness);
+				constexpr bool bCanTimeout = true;
+				WaitForSample(NextFrameExpectations, MoveTemp(ReadyToRenderConditionFunc), bCanTimeout);
+			});
 		}
 		else
 		{
@@ -1139,17 +1137,15 @@ namespace UE::RivermaxMedia
 			}
 			
 			// Setup requirements for sample to be ready to be rendered
-			FRHICommandListImmediate& RHICmdList = FRHICommandListExecutor::GetImmediateCommandList();
-
-			SCOPED_GPU_STAT(RHICmdList, Rmax_WaitForPixels);
-			SCOPED_DRAW_EVENT(RHICmdList, Rmax_WaitForPixels);
+			SCOPED_GPU_STAT(GraphBuilder.RHICmdList, Rmax_WaitForPixels);
+			SCOPED_DRAW_EVENT(GraphBuilder.RHICmdList, Rmax_WaitForPixels);
 
 			// Since we are going to enqueue a lambda that can potentially sleep in the RHI thread if the pixels haven't arrived,
 			// we dispatch the existing commands (including the draw event start timing in the SCOPED_DRAW_EVENT above) before any potential sleep.
-			RHICmdList.ImmediateFlush(EImmediateFlushType::DispatchToRHIThread);
+			GraphBuilder.RHICmdList.ImmediateFlush(EImmediateFlushType::DispatchToRHIThread);
 
-			RHICmdList.EnqueueLambda(
-				[NextFrameExpectations, this](FRHICommandList& RHICmdList)
+			GraphBuilder.RHICmdList.EnqueueLambda(
+				[NextFrameExpectations, this](FRHICommandListImmediate&)
 				{
 					FWaitConditionFunc ReadyToRenderConditionFunc = [](const TSharedPtr<FRivermaxSampleWrapper>& Sample)
 					{
@@ -1168,7 +1164,7 @@ namespace UE::RivermaxMedia
 			// Final step, if the memory was locked (non gpu direct), enqueue unlock after the wait for sample in order to render it
 			if (SampleWrapper->LockedMemory && ensure(!bDoesStreamSupportsGPUDirect))
 			{
-				RHICmdList.UnlockBuffer(SampleWrapper->Sample->GetGPUBuffer()->GetRHI());
+				GraphBuilder.RHICmdList.UnlockBuffer(SampleWrapper->Sample->GetGPUBuffer()->GetRHI());
 			}
 		};
 	}
diff --git a/Engine/Plugins/VirtualProduction/Rivermax/RivermaxMedia/Source/RivermaxMedia/Private/RivermaxMediaTextureSampleConverter.cpp b/Engine/Plugins/VirtualProduction/Rivermax/RivermaxMedia/Source/RivermaxMedia/Private/RivermaxMediaTextureSampleConverter.cpp
--- a/Engine/Plugins/VirtualProduction/Rivermax/RivermaxMedia/Source/RivermaxMedia/Private/RivermaxMediaTextureSampleConverter.cpp
+++ b/Engine/Plugins/VirtualProduction/Rivermax/RivermaxMedia/Source/RivermaxMedia/Private/RivermaxMediaTextureSampleConverter.cpp
@@ -19,7 +19,7 @@ void FRivermaxMediaTextureSampleConverter::Setup(const TSharedPtr<FRivermaxMedia
 	Sample = InSample;
 }
 
-bool FRivermaxMediaTextureSampleConverter::Convert(FTextureRHIRef& InDestinationTexture, const FConversionHints& Hints)
+bool FRivermaxMediaTextureSampleConverter::Convert(FRHICommandListImmediate& RHICmdList, FTextureRHIRef& InDestinationTexture, const FConversionHints& Hints)
 {
 	TRACE_CPUPROFILER_EVENT_SCOPE(RivermaxSampleConverter::Convert);
 
@@ -53,7 +53,7 @@ bool FRivermaxMediaTextureSampleConverter::Convert(FTextureRHIRef& InDestination
 		return false;
 	}
 	
-	FRDGBuilder GraphBuilder(FRHICommandListExecutor::GetImmediateCommandList());
+	FRDGBuilder GraphBuilder(RHICmdList);
 	if (Setup.PreConvertFunc)
 	{
 		Setup.PreConvertFunc(GraphBuilder);
diff --git a/Engine/Plugins/VirtualProduction/Rivermax/RivermaxMedia/Source/RivermaxMedia/Private/RivermaxMediaTextureSampleConverter.h b/Engine/Plugins/VirtualProduction/Rivermax/RivermaxMedia/Source/RivermaxMedia/Private/RivermaxMediaTextureSampleConverter.h
--- a/Engine/Plugins/VirtualProduction/Rivermax/RivermaxMedia/Source/RivermaxMedia/Private/RivermaxMediaTextureSampleConverter.h
+++ b/Engine/Plugins/VirtualProduction/Rivermax/RivermaxMedia/Source/RivermaxMedia/Private/RivermaxMediaTextureSampleConverter.h
@@ -56,7 +56,7 @@ public:
 	void Setup(const TSharedPtr<FRivermaxMediaTextureSample>& InSample);
 
 	//~ Begin IMediaTextureSampleConverter interface
-	virtual bool Convert(FTextureRHIRef& InDstTexture, const FConversionHints& Hints) override;
+	virtual bool Convert(FRHICommandListImmediate& RHICmdList, FTextureRHIRef& InDstTexture, const FConversionHints& Hints) override;
 	virtual uint32 GetConverterInfoFlags() const override;
 	//~ End IMediaTextureSampleConverter interface
 
diff --git a/Engine/Plugins/VirtualProduction/Rivermax/RivermaxMedia/Source/RivermaxMedia/Public/RivermaxMediaCapture.h b/Engine/Plugins/VirtualProduction/Rivermax/RivermaxMedia/Source/RivermaxMedia/Public/RivermaxMediaCapture.h
--- a/Engine/Plugins/VirtualProduction/Rivermax/RivermaxMedia/Source/RivermaxMedia/Public/RivermaxMediaCapture.h
+++ b/Engine/Plugins/VirtualProduction/Rivermax/RivermaxMedia/Source/RivermaxMedia/Public/RivermaxMediaCapture.h
@@ -41,7 +41,7 @@ protected:
 	virtual void StopCaptureImpl(bool bAllowPendingFrameToBeProcess) override;
 	virtual bool ShouldCaptureRHIResource() const override;
 	virtual void OnFrameCaptured_RenderingThread(const FCaptureBaseData& InBaseData, TSharedPtr<FMediaCaptureUserData, ESPMode::ThreadSafe> InUserData, void* InBuffer, int32 Width, int32 Height, int32 BytesPerRow) override;
-	virtual void OnRHIResourceCaptured_RenderingThread(const FCaptureBaseData& InBaseData, TSharedPtr<FMediaCaptureUserData, ESPMode::ThreadSafe> InUserData, FBufferRHIRef InBuffer) override;
+	virtual void OnRHIResourceCaptured_RenderingThread(FRHICommandListImmediate& RHICmdList, const FCaptureBaseData& InBaseData, TSharedPtr<FMediaCaptureUserData, ESPMode::ThreadSafe> InUserData, FBufferRHIRef InBuffer) override;
 	virtual void OnRHIResourceCaptured_AnyThread(const FCaptureBaseData& InBaseData, TSharedPtr<FMediaCaptureUserData, ESPMode::ThreadSafe> InUserData, FBufferRHIRef InBuffer) override;
 	virtual void OnFrameCaptured_AnyThread(const FCaptureBaseData& InBaseData, TSharedPtr<FMediaCaptureUserData, ESPMode::ThreadSafe> InUserData, const FMediaCaptureResourceData& InResourceData) override;
 	virtual bool SupportsAnyThreadCapture() const override
diff --git a/Engine/Source/Runtime/GameplayMediaEncoder/Private/GameplayMediaEncoder.cpp b/Engine/Source/Runtime/GameplayMediaEncoder/Private/GameplayMediaEncoder.cpp
--- a/Engine/Source/Runtime/GameplayMediaEncoder/Private/GameplayMediaEncoder.cpp
+++ b/Engine/Source/Runtime/GameplayMediaEncoder/Private/GameplayMediaEncoder.cpp
@@ -780,7 +780,7 @@ PRAGMA_ENABLE_DEPRECATION_WARNINGS
 
 void FGameplayMediaEncoder::CopyTexture(const FTextureRHIRef& SourceTexture, FTextureRHIRef& DestinationTexture) const
 {
-	FRHICommandListImmediate& RHICmdList = FRHICommandListExecutor::GetImmediateCommandList();
+	FRHICommandListImmediate& RHICmdList = FRHICommandListImmediate::Get();
 
 	if(SourceTexture->GetFormat() == DestinationTexture->GetFormat() && SourceTexture->GetSizeXY() == DestinationTexture->GetSizeXY())
 	{
diff --git a/Engine/Source/Runtime/Media/Public/IMediaTextureSampleConverter.h b/Engine/Source/Runtime/Media/Public/IMediaTextureSampleConverter.h
--- a/Engine/Source/Runtime/Media/Public/IMediaTextureSampleConverter.h
+++ b/Engine/Source/Runtime/Media/Public/IMediaTextureSampleConverter.h
@@ -6,6 +6,7 @@
 
 #include "RHI.h"
 #include "RHIResources.h"
+#include "RHICommandList.h"
 
 /**
  * Interface class to implement custom sample conversion
@@ -33,7 +34,13 @@ public:
 		return ConverterInfoFlags_Default;
 	}
 
-	virtual bool Convert(FTextureRHIRef & InDstTexture, const FConversionHints & Hints) = 0;
+	virtual bool Convert(FRHICommandListImmediate& RHICmdList, FTextureRHIRef& InDstTexture, const FConversionHints& Hints) = 0;
+
+	UE_DEPRECATED(5.5, "The overload of IMediaTextureSampleConverter::Convert which does not take an RHI command list is deprecated.")
+	bool Convert(FTextureRHIRef& InDstTexture, const FConversionHints& Hints)
+	{
+		return Convert(FRHICommandListImmediate::Get(), InDstTexture, Hints);
+	}
 };
 
 /**
@@ -48,6 +55,12 @@ public:
 	 * Apply a color conversion on the input and store the result in the destination texture.
 	 * @return true If the color conversion was successfully applied.
 	 */
-	virtual bool ApplyColorConversion(FTextureRHIRef& InSrcTexture, FTextureRHIRef& InDstTexture) = 0;
+	virtual bool ApplyColorConversion(FRHICommandListImmediate& RHICmdList, FTextureRHIRef& InSrcTexture, FTextureRHIRef& InDstTexture) = 0;
+
+	UE_DEPRECATED(5.5, "The overload of IMediaTextureSampleColorConverter::ApplyColorConversion which does not take an RHI command list is deprecated.")
+	bool ApplyColorConversion(FTextureRHIRef& InSrcTexture, FTextureRHIRef& InDstTexture)
+	{
+		return ApplyColorConversion(FRHICommandListImmediate::Get(), InSrcTexture, InDstTexture);
+	}
 };
 
diff --git a/Engine/Source/Runtime/MediaAssets/Private/Misc/MediaTextureResource.cpp b/Engine/Source/Runtime/MediaAssets/Private/Misc/MediaTextureResource.cpp
--- a/Engine/Source/Runtime/MediaAssets/Private/Misc/MediaTextureResource.cpp
+++ b/Engine/Source/Runtime/MediaAssets/Private/Misc/MediaTextureResource.cpp
@@ -77,7 +77,7 @@ template<typename ObjectRefType> struct TGPUsyncedDataDeleter : public IMediaClo
 		}
 	}
 
-	void Retire(const ObjectRefType& Object)
+	void Retire(FRHICommandListImmediate& RHICmdList, const ObjectRefType& Object)
 	{
 		// Prep "retirement package"
 		FRetiringObjectInfo Info;
@@ -86,7 +86,7 @@ template<typename ObjectRefType> struct TGPUsyncedDataDeleter : public IMediaClo
 		Info.RetireTime = FPlatformTime::Seconds();
 
 		// Insert fence. We assume that GPU-workload-wise this marks the spot usage of the sample is done
-		FRHICommandListExecutor::GetImmediateCommandList().WriteGPUFence(Info.GPUFence);
+		RHICmdList.WriteGPUFence(Info.GPUFence);
 
 		// Recall for later checking...
 		FScopeLock Lock(&CS);
@@ -457,7 +457,7 @@ void FMediaTextureResource::Render(const FRenderParams& Params)
 				//
 				// Sample dimensions are invalid
 				//
-				ClearTexture(FLinearColor::Red, false); // mark corrupt sample
+				ClearTexture(RHICmdList, FLinearColor::Red, false); // mark corrupt sample
 			}
 			else if (IMediaTextureSampleConverter *Converter = Sample->GetMediaTextureSampleConverter())
 			{
@@ -480,7 +480,7 @@ void FMediaTextureResource::Render(const FRenderParams& Params)
 					{
 						// Preprocess...
 						FTextureRHIRef DummyTexture;
-						if (Converter->Convert(DummyTexture, Hints))
+						if (Converter->Convert(RHICmdList, DummyTexture, Hints))
 						{
 							// ...followed by the built in conversion code as needed...
 							ConvertOrCopyNeeded = true;
@@ -490,25 +490,25 @@ void FMediaTextureResource::Render(const FRenderParams& Params)
 					{
 						// Conversion is fully handled by converter
 
-						CreateIntermediateRenderTarget(Sample->GetOutputDim(), MediaTextureResourceHelpers::GetConvertedPixelFormat(Sample), MediaTextureResourceHelpers::RequiresSrgbTexture(Sample), Params.ClearColor, Hints.NumMips, bNeedsUAVTexture);
-						Converter->Convert(RenderTargetTextureRHI, Hints);
+						CreateIntermediateRenderTarget(RHICmdList, Sample->GetOutputDim(), MediaTextureResourceHelpers::GetConvertedPixelFormat(Sample), MediaTextureResourceHelpers::RequiresSrgbTexture(Sample), Params.ClearColor, Hints.NumMips, bNeedsUAVTexture);
+						Converter->Convert(RHICmdList, RenderTargetTextureRHI, Hints);
 					}
 				}
 				else
 				{
 					// The converter will create its own output texture for us to use
 					FTextureRHIRef OutTexture;
-					if (Converter->Convert(OutTexture, Hints))
+					if (Converter->Convert(RHICmdList, OutTexture, Hints))
 					{
 						// As the converter created the texture, we might need to convert it even more to make it fit our needs. Check...
 						if (RequiresConversion(OutTexture, Sample->GetOutputDim(), NumMips))
 						{
-							CreateIntermediateRenderTarget(Sample->GetOutputDim(), MediaTextureResourceHelpers::GetConvertedPixelFormat(Sample), MediaTextureResourceHelpers::RequiresSrgbTexture(Sample), Params.ClearColor, Hints.NumMips, bNeedsUAVTexture);
-							ConvertTextureToOutput(OutTexture.GetReference(), Sample);
+							CreateIntermediateRenderTarget(RHICmdList, Sample->GetOutputDim(), MediaTextureResourceHelpers::GetConvertedPixelFormat(Sample), MediaTextureResourceHelpers::RequiresSrgbTexture(Sample), Params.ClearColor, Hints.NumMips, bNeedsUAVTexture);
+							ConvertTextureToOutput(RHICmdList, OutTexture.GetReference(), Sample);
 						}
 						else
 						{
-							UpdateTextureReference(OutTexture);
+							UpdateTextureReference(RHICmdList, OutTexture);
 						}
 					}
 				}
@@ -539,14 +539,14 @@ void FMediaTextureResource::Render(const FRenderParams& Params)
 						//
 						// Sample needs to be converted by built in converter code
 						//
-						ConvertSample(Sample, Params.ClearColor, NumMips);
+						ConvertSample(RHICmdList, Sample, Params.ClearColor, NumMips);
 					}
 					else
 					{
 						//
 						// Sample can be used directly or is a simple copy
 						//
-						CopySample(Sample, Params.ClearColor, NumMips, Params.CurrentGuid);
+						CopySample(RHICmdList, Sample, Params.ClearColor, NumMips, Params.CurrentGuid);
 					}
 				}
 
@@ -571,18 +571,18 @@ void FMediaTextureResource::Render(const FRenderParams& Params)
 							.SetOwnerName(GetOwnerName())
 							.SetFlags(IntermediateTextureDesc.Flags);
 
-						OutputTarget = RHICreateTexture(Desc);
+						OutputTarget = RHICmdList.CreateTexture(Desc);
 						OutputTarget->SetName(TEXT("MediaTextureResourceOutput"));
 						OutputTarget->SetOwnerName(GetOwnerName());
 
 						bRecreateOutputTarget = false;
 					}
 
-					Converter->ApplyColorConversion(TextureRef, OutputTarget);
+					Converter->ApplyColorConversion(RHICmdList, TextureRef, OutputTarget);
 
 					if (RenderTargetTextureRHI != OutputTarget)
 					{
-						UpdateTextureReference(OutputTarget);
+						UpdateTextureReference(RHICmdList, OutputTarget);
 					}
 				}
 				else if (IntermediateTarget)
@@ -596,7 +596,7 @@ void FMediaTextureResource::Render(const FRenderParams& Params)
 				if (CurrentSample)
 				{
 					// If we had a current sample (directly used as output), we can now schedule its retirement
-					LocalPriorSamples->Retire(CurrentSample);
+					LocalPriorSamples->Retire(RHICmdList, CurrentSample);
 					CurrentSample = nullptr;
 				}
 
@@ -604,7 +604,7 @@ void FMediaTextureResource::Render(const FRenderParams& Params)
 				if (OutputTarget == RenderTargetTextureRHI)
 				{
 					// Yes, we can schedule the actual sample for retirement right away
-					LocalPriorSamples->Retire(Sample);
+					LocalPriorSamples->Retire(RHICmdList, Sample);
 				}
 				else
 				{
@@ -642,12 +642,12 @@ void FMediaTextureResource::Render(const FRenderParams& Params)
 		if (!Cleared || (Params.ClearColor != CurrentClearColor))
 		{
 			// Yes...
-			ClearTexture(Params.ClearColor, false);
+			ClearTexture(RHICmdList, Params.ClearColor, false);
 
 			if (CurrentSample)
 			{
 				// If we had a current sample (directly used as output), we can now schedule its retirement
-				LocalPriorSamples->Retire(CurrentSample);
+				LocalPriorSamples->Retire(RHICmdList, CurrentSample);
 				CurrentSample = nullptr;
 			}
 		}
@@ -736,14 +736,16 @@ void FMediaTextureResource::SetupSampler()
 }
 
 
-void FMediaTextureResource::InitRHI(FRHICommandListBase&)
+void FMediaTextureResource::InitRHI(FRHICommandListBase& RHICmdListBase)
 {
+	FRHICommandListImmediate& RHICmdList = RHICmdListBase.GetAsImmediate();
+
 	SetupSampler();
 	
 	// Note: set up default texture, or we can get sampler bind errors on render
 	// we can't leave here without having a valid bindable resource for some RHIs.
 
-	ClearTexture(CurrentClearColor, Owner.SRGB);
+	ClearTexture(RHICmdList, CurrentClearColor, Owner.SRGB);
 
 	// Make sure init has done it's job - we can't leave here without valid bindable resources for some RHI's
 	check(TextureRHI.IsValid());
@@ -771,7 +773,7 @@ void FMediaTextureResource::ReleaseRHI()
 	RenderTargetTextureRHI.SafeRelease();
 	TextureRHI.SafeRelease();
 
-	UpdateTextureReference(nullptr);
+	UpdateTextureReference(FRHICommandListImmediate::Get(), nullptr);
 }
 
 
@@ -805,25 +807,24 @@ void FMediaTextureResource::JustInTimeRender()
 }
 
 
-void FMediaTextureResource::ClearTexture(const FLinearColor& ClearColor, bool SrgbOutput)
+void FMediaTextureResource::ClearTexture(FRHICommandListImmediate& RHICmdList, const FLinearColor& ClearColor, bool SrgbOutput)
 {
 	// create output render target if we don't have one yet
 	constexpr uint8 NumMips = 1;
 	constexpr bool bNeedsUAVTexture = false;
-	CreateIntermediateRenderTarget(FIntPoint(2, 2), PF_B8G8R8A8, SrgbOutput, ClearColor, NumMips, bNeedsUAVTexture);
+	CreateIntermediateRenderTarget(RHICmdList, FIntPoint(2, 2), PF_B8G8R8A8, SrgbOutput, ClearColor, NumMips, bNeedsUAVTexture);
 
 	// draw the clear color
-	FRHICommandListImmediate& CommandList = FRHICommandListExecutor::GetImmediateCommandList();
 	{
-		SCOPED_DRAW_EVENT(CommandList, FMediaTextureResource_ClearTexture);
-		SCOPED_GPU_STAT(CommandList, MediaTextureResource);
+		SCOPED_DRAW_EVENT(RHICmdList, FMediaTextureResource_ClearTexture);
+		SCOPED_GPU_STAT(RHICmdList, MediaTextureResource);
 
-		CommandList.Transition(FRHITransitionInfo(RenderTargetTextureRHI, ERHIAccess::SRVMask, ERHIAccess::RTV));
+		RHICmdList.Transition(FRHITransitionInfo(RenderTargetTextureRHI, ERHIAccess::SRVMask, ERHIAccess::RTV));
 
 		FRHIRenderPassInfo RPInfo(RenderTargetTextureRHI, ERenderTargetActions::Clear_Store);
-		CommandList.BeginRenderPass(RPInfo, TEXT("ClearTexture"));
-		CommandList.EndRenderPass();
-		CommandList.Transition(FRHITransitionInfo(RenderTargetTextureRHI, ERHIAccess::RTV, ERHIAccess::SRVMask));
+		RHICmdList.BeginRenderPass(RPInfo, TEXT("ClearTexture"));
+		RHICmdList.EndRenderPass();
+		RHICmdList.Transition(FRHITransitionInfo(RenderTargetTextureRHI, ERHIAccess::RTV, ERHIAccess::SRVMask));
 	}
 
 	Cleared = true;
@@ -949,7 +950,7 @@ bool FMediaTextureResource::RequiresConversion(const TSharedPtr<IMediaTextureSam
 
 
 
-void FMediaTextureResource::ConvertSample(const TSharedPtr<IMediaTextureSample, ESPMode::ThreadSafe>& Sample, const FLinearColor& ClearColor, uint8 InNumMips)
+void FMediaTextureResource::ConvertSample(FRHICommandListImmediate& RHICmdList, const TSharedPtr<IMediaTextureSample, ESPMode::ThreadSafe>& Sample, const FLinearColor& ClearColor, uint8 InNumMips)
 {
 	const EPixelFormat InputPixelFormat = MediaTextureResourceHelpers::GetPixelFormat(Sample);
 
@@ -988,7 +989,7 @@ void FMediaTextureResource::ConvertSample(const TSharedPtr<IMediaTextureSample,
 					.SetNumMips(SampleNumMips)
 					.SetFlags(InputCreateFlags);
 
-				InputTarget = RHICreateTexture(Desc);
+				InputTarget = RHICmdList.CreateTexture(Desc);
 
 				UpdateResourceSize();
 			}
@@ -1001,7 +1002,7 @@ void FMediaTextureResource::ConvertSample(const TSharedPtr<IMediaTextureSample,
 				uint32 Stride = Sample->GetStride() >> MipLevel;
 				uint32 Height = SampleDim.Y >> MipLevel;
 				FUpdateTextureRegion2D Region(0, 0, 0, 0, SampleDim.X >> MipLevel, Height);
-				RHIUpdateTexture2D(InputTarget, MipLevel, Region, Stride, Data);
+				RHICmdList.UpdateTexture2D(InputTarget, MipLevel, Region, Stride, Data);
 				Data += Stride * Height;
 			}
 
@@ -1013,9 +1014,9 @@ void FMediaTextureResource::ConvertSample(const TSharedPtr<IMediaTextureSample,
 	const FIntPoint OutputDim = Sample->GetOutputDim();
 	const uint8 NumMips = (SampleNumMips > 1) ? SampleNumMips : InNumMips;
 	const bool bNeedsUAVTexture = MediaTextureResourceHelpers::RequiresUAVTexture(Sample, NumMips);
-	CreateIntermediateRenderTarget(OutputDim, MediaTextureResourceHelpers::GetConvertedPixelFormat(Sample), MediaTextureResourceHelpers::RequiresSrgbTexture(Sample), ClearColor, NumMips, bNeedsUAVTexture);
+	CreateIntermediateRenderTarget(RHICmdList, OutputDim, MediaTextureResourceHelpers::GetConvertedPixelFormat(Sample), MediaTextureResourceHelpers::RequiresSrgbTexture(Sample), ClearColor, NumMips, bNeedsUAVTexture);
 
-	ConvertTextureToOutput(InputTexture, Sample);
+	ConvertTextureToOutput(RHICmdList, InputTexture, Sample);
 }
 
 
@@ -1040,32 +1041,31 @@ void FMediaTextureResource::GetColorSpaceConversionMatrixForSample(const TShared
 }
 
 
- void FMediaTextureResource::ConvertTextureToOutput(FRHITexture* InputTexture, const TSharedPtr<IMediaTextureSample, ESPMode::ThreadSafe>& Sample)
+ void FMediaTextureResource::ConvertTextureToOutput(FRHICommandListImmediate& RHICmdList, FRHITexture* InputTexture, const TSharedPtr<IMediaTextureSample, ESPMode::ThreadSafe>& Sample)
  {
 	// perform the conversion
-	FRHICommandListImmediate& CommandList = FRHICommandListExecutor::GetImmediateCommandList();
 	{
 		// We should never get here with a sample that contains mips!
 		check(Sample->GetNumMips() == 1);
 
-		SCOPED_DRAW_EVENT(CommandList, FMediaTextureResource_Convert);
-		SCOPED_GPU_STAT(CommandList, MediaTextureResource);
+		SCOPED_DRAW_EVENT(RHICmdList, FMediaTextureResource_Convert);
+		SCOPED_GPU_STAT(RHICmdList, MediaTextureResource);
 
 		FGraphicsPipelineStateInitializer GraphicsPSOInit;
 		FRHITexture* RenderTarget = IntermediateTarget;
-		CommandList.Transition(FRHITransitionInfo(RenderTarget, ERHIAccess::Unknown, ERHIAccess::RTV));
+		RHICmdList.Transition(FRHITransitionInfo(RenderTarget, ERHIAccess::Unknown, ERHIAccess::RTV));
 
 		FIntPoint OutputDim(RenderTarget->GetSizeXYZ().X, RenderTarget->GetSizeXYZ().Y);
 
 		// note: we are not explicitly transitioning the input texture to be readable here
 		// (we assume this to be the case already - main as some platforms may fail to orderly transition the resource due to special cases regarding their internal setup)
-		CommandList.Transition(FRHITransitionInfo(RenderTargetTextureRHI, ERHIAccess::Unknown, ERHIAccess::RTV));
+		RHICmdList.Transition(FRHITransitionInfo(RenderTargetTextureRHI, ERHIAccess::Unknown, ERHIAccess::RTV));
 
 		FRHIRenderPassInfo RPInfo(RenderTarget, ERenderTargetActions::DontLoad_Store);
-		CommandList.BeginRenderPass(RPInfo, TEXT("ConvertMedia"));
+		RHICmdList.BeginRenderPass(RPInfo, TEXT("ConvertMedia"));
 		{
-			CommandList.ApplyCachedRenderTargets(GraphicsPSOInit);
-			CommandList.SetViewport(0, 0, 0.0f, (float)OutputDim.X, (float)OutputDim.Y, 1.0f);
+			RHICmdList.ApplyCachedRenderTargets(GraphicsPSOInit);
+			RHICmdList.SetViewport(0, 0, 0.0f, (float)OutputDim.X, (float)OutputDim.Y, 1.0f);
 
 			GraphicsPSOInit.DepthStencilState = TStaticDepthStencilState<false, CF_Always>::GetRHI();
 			GraphicsPSOInit.RasterizerState = TStaticRasterizerState<>::GetRHI();
@@ -1098,18 +1098,18 @@ void FMediaTextureResource::GetColorSpaceConversionMatrixForSample(const TShared
 					{
 						TShaderMapRef<FNV12ConvertPS> ConvertShader(ShaderMap);
 						GraphicsPSOInit.BoundShaderState.PixelShaderRHI = ConvertShader.GetPixelShader();
-						SetGraphicsPipelineState(CommandList, GraphicsPSOInit, 0);
+						SetGraphicsPipelineState(RHICmdList, GraphicsPSOInit, 0);
 						FIntPoint TexDim = InputTexture->GetSizeXY();
-						TempSRV0 = CommandList.CreateShaderResourceView(InputTexture, 0, 1, PF_G8);								// note: RHI does provide "magic" to select Y vs. UV planes based on the pixel format (D3D/DXGI)
-						TempSRV1 = CommandList.CreateShaderResourceView(InputTexture, 0, 1, PF_R8G8);
-						SetShaderParametersLegacyPS(CommandList, ConvertShader, TexDim, TempSRV0, TempSRV1, OutputDim, YUVMtx, Sample->GetEncodingType(), ColorSpaceMtx, SampleFormat == EMediaTextureSampleFormat::CharNV21);
+						TempSRV0 = RHICmdList.CreateShaderResourceView(InputTexture, 0, 1, PF_G8);								// note: RHI does provide "magic" to select Y vs. UV planes based on the pixel format (D3D/DXGI)
+						TempSRV1 = RHICmdList.CreateShaderResourceView(InputTexture, 0, 1, PF_R8G8);
+						SetShaderParametersLegacyPS(RHICmdList, ConvertShader, TexDim, TempSRV0, TempSRV1, OutputDim, YUVMtx, Sample->GetEncodingType(), ColorSpaceMtx, SampleFormat == EMediaTextureSampleFormat::CharNV21);
 					}
 					else
 					{
 						TShaderMapRef<FNV12ConvertAsBytesPS> ConvertShader(ShaderMap);
 						GraphicsPSOInit.BoundShaderState.PixelShaderRHI = ConvertShader.GetPixelShader();
-						SetGraphicsPipelineState(CommandList, GraphicsPSOInit, 0);
-						SetShaderParametersLegacyPS(CommandList, ConvertShader, InputTexture, OutputDim, YUVMtx, Sample->GetEncodingType(), ColorSpaceMtx, SampleFormat == EMediaTextureSampleFormat::CharNV21);
+						SetGraphicsPipelineState(RHICmdList, GraphicsPSOInit, 0);
+						SetShaderParametersLegacyPS(RHICmdList, ConvertShader, InputTexture, OutputDim, YUVMtx, Sample->GetEncodingType(), ColorSpaceMtx, SampleFormat == EMediaTextureSampleFormat::CharNV21);
 					}
 				}
 				break;
@@ -1126,19 +1126,19 @@ void FMediaTextureResource::GetColorSpaceConversionMatrixForSample(const TShared
 					{
 						TShaderMapRef<FP010ConvertPS> ConvertShader(ShaderMap);
 						GraphicsPSOInit.BoundShaderState.PixelShaderRHI = ConvertShader.GetPixelShader();
-						SetGraphicsPipelineState(CommandList, GraphicsPSOInit, 0);
+						SetGraphicsPipelineState(RHICmdList, GraphicsPSOInit, 0);
 
-						FShaderResourceViewRHIRef Y_SRV = CommandList.CreateShaderResourceView(InputTexture, 0, 1, PF_G16);		// note: RHI does provide "magic" to select Y vs. UV planes based on the pixel format (D3D/DXGI)
-						FShaderResourceViewRHIRef UV_SRV = CommandList.CreateShaderResourceView(InputTexture, 0, 1, PF_G16R16);
-						SetShaderParametersLegacyPS(CommandList, ConvertShader, TexDim, Y_SRV, UV_SRV, OutputDim, YUVMtx, ColorSpaceMtx, Sample->GetEncodingType());
+						FShaderResourceViewRHIRef Y_SRV = RHICmdList.CreateShaderResourceView(InputTexture, 0, 1, PF_G16);		// note: RHI does provide "magic" to select Y vs. UV planes based on the pixel format (D3D/DXGI)
+						FShaderResourceViewRHIRef UV_SRV = RHICmdList.CreateShaderResourceView(InputTexture, 0, 1, PF_G16R16);
+						SetShaderParametersLegacyPS(RHICmdList, ConvertShader, TexDim, Y_SRV, UV_SRV, OutputDim, YUVMtx, ColorSpaceMtx, Sample->GetEncodingType());
 					}
 					else
 					{
 						TShaderMapRef<FP010ConvertAsUINT16sPS> ConvertShader(ShaderMap);
 						GraphicsPSOInit.BoundShaderState.PixelShaderRHI = ConvertShader.GetPixelShader();
-						SetGraphicsPipelineState(CommandList, GraphicsPSOInit, 0);
+						SetGraphicsPipelineState(RHICmdList, GraphicsPSOInit, 0);
 
-						SetShaderParametersLegacyPS(CommandList, ConvertShader, TexDim, InputTexture, OutputDim, YUVMtx, ColorSpaceMtx, Sample->GetEncodingType());
+						SetShaderParametersLegacyPS(RHICmdList, ConvertShader, TexDim, InputTexture, OutputDim, YUVMtx, ColorSpaceMtx, Sample->GetEncodingType());
 					}
 				}
 				break;
@@ -1156,8 +1156,8 @@ void FMediaTextureResource::GetColorSpaceConversionMatrixForSample(const TShared
 
 					TShaderMapRef<FYUVv216ConvertPS> ConvertShader(ShaderMap);
 					GraphicsPSOInit.BoundShaderState.PixelShaderRHI = ConvertShader.GetPixelShader();
-					SetGraphicsPipelineState(CommandList, GraphicsPSOInit, 0);
-					SetShaderParametersLegacyPS(CommandList, ConvertShader, InputTexture, OutputDim, YUVMtx, Sample->GetEncodingType(), ColorSpaceMtx,
+					SetGraphicsPipelineState(RHICmdList, GraphicsPSOInit, 0);
+					SetShaderParametersLegacyPS(RHICmdList, ConvertShader, InputTexture, OutputDim, YUVMtx, Sample->GetEncodingType(), ColorSpaceMtx,
 												SampleFormat!= EMediaTextureSampleFormat::CharYUY2 && SampleFormat != EMediaTextureSampleFormat::CharYVYU,	// Y or Cb first
 												InputTexture->GetFormat() == PF_B8G8R8A8,																	// ARGB vs. ABGR (memory order)
 												SampleFormat == EMediaTextureSampleFormat::CharYVYU															// Cb / Cr swap
@@ -1174,8 +1174,8 @@ void FMediaTextureResource::GetColorSpaceConversionMatrixForSample(const TShared
 
 					TShaderMapRef<FYUVv210ConvertPS> ConvertShader(ShaderMap);
 					GraphicsPSOInit.BoundShaderState.PixelShaderRHI = ConvertShader.GetPixelShader();
-					SetGraphicsPipelineState(CommandList, GraphicsPSOInit, 0);
-					SetShaderParametersLegacyPS(CommandList, ConvertShader, InputTexture, OutputDim, YUVMtx, Sample->GetEncodingType(), ColorSpaceMtx,
+					SetGraphicsPipelineState(RHICmdList, GraphicsPSOInit, 0);
+					SetShaderParametersLegacyPS(RHICmdList, ConvertShader, InputTexture, OutputDim, YUVMtx, Sample->GetEncodingType(), ColorSpaceMtx,
 												true);
 				}
 				break;
@@ -1191,10 +1191,10 @@ void FMediaTextureResource::GetColorSpaceConversionMatrixForSample(const TShared
 
 					TShaderMapRef<FYUVY416ConvertPS> ConvertShader(ShaderMap);
 					GraphicsPSOInit.BoundShaderState.PixelShaderRHI = ConvertShader.GetPixelShader();
-					SetGraphicsPipelineState(CommandList, GraphicsPSOInit, 0);
-					FShaderResourceViewRHIRef SRV = CommandList.CreateShaderResourceView(InputTexture, 0, 1, (Sample->GetFormat() == EMediaTextureSampleFormat::Y416) ? PF_A16B16G16R16 : PF_A32B32G32R32F);
+					SetGraphicsPipelineState(RHICmdList, GraphicsPSOInit, 0);
+					FShaderResourceViewRHIRef SRV = RHICmdList.CreateShaderResourceView(InputTexture, 0, 1, (Sample->GetFormat() == EMediaTextureSampleFormat::Y416) ? PF_A16B16G16R16 : PF_A32B32G32R32F);
 
-					SetShaderParametersLegacyPS(CommandList, ConvertShader, SRV, YUVMtx, Sample->GetEncodingType(), ColorSpaceMtx,
+					SetShaderParametersLegacyPS(RHICmdList, ConvertShader, SRV, YUVMtx, Sample->GetEncodingType(), ColorSpaceMtx,
 												InputTexture->GetFormat() == PF_A8R8G8B8);
 				}
 				break;
@@ -1227,8 +1227,8 @@ void FMediaTextureResource::GetColorSpaceConversionMatrixForSample(const TShared
 
 					TShaderMapRef<FRGBConvertPS> ConvertShader(ShaderMap);
 					GraphicsPSOInit.BoundShaderState.PixelShaderRHI = ConvertShader.GetPixelShader();
-					SetGraphicsPipelineState(CommandList, GraphicsPSOInit, 0);
-					SetShaderParametersLegacyPS(CommandList, ConvertShader, InputTexture, OutputDim, Encoding, ColorSpaceMtx);
+					SetGraphicsPipelineState(RHICmdList, GraphicsPSOInit, 0);
+					SetShaderParametersLegacyPS(RHICmdList, ConvertShader, InputTexture, OutputDim, Encoding, ColorSpaceMtx);
 				}
 				break;
 
@@ -1239,10 +1239,10 @@ void FMediaTextureResource::GetColorSpaceConversionMatrixForSample(const TShared
 
 					TShaderMapRef<FARGB16BigConvertPS> ConvertShader(ShaderMap);
 					GraphicsPSOInit.BoundShaderState.PixelShaderRHI = ConvertShader.GetPixelShader();
-					SetGraphicsPipelineState(CommandList, GraphicsPSOInit, 0);
-					FShaderResourceViewRHIRef SRV = CommandList.CreateShaderResourceView(InputTexture, 0, 1, PF_R16G16B16A16_UINT);
+					SetGraphicsPipelineState(RHICmdList, GraphicsPSOInit, 0);
+					FShaderResourceViewRHIRef SRV = RHICmdList.CreateShaderResourceView(InputTexture, 0, 1, PF_R16G16B16A16_UINT);
 
-					SetShaderParametersLegacyPS(CommandList, ConvertShader, SRV, OutputDim, Sample->GetEncodingType(), ColorSpaceMtx);
+					SetShaderParametersLegacyPS(RHICmdList, ConvertShader, SRV, OutputDim, Sample->GetEncodingType(), ColorSpaceMtx);
 				}
 				break;
 
@@ -1251,8 +1251,8 @@ void FMediaTextureResource::GetColorSpaceConversionMatrixForSample(const TShared
 					// Simple 1:1 copy plus flip & color adjustment (but using normal texture sampler: sRGB conversions may occur depending on setup; any manual sRGB/linear conversion is disabled)
 					TShaderMapRef<FBMPConvertPS> ConvertShader(ShaderMap);
 					GraphicsPSOInit.BoundShaderState.PixelShaderRHI = ConvertShader.GetPixelShader();
-					SetGraphicsPipelineState(CommandList, GraphicsPSOInit, 0);
-					SetShaderParametersLegacyPS(CommandList, ConvertShader, InputTexture, OutputDim, false);
+					SetGraphicsPipelineState(RHICmdList, GraphicsPSOInit, 0);
+					SetShaderParametersLegacyPS(RHICmdList, ConvertShader, InputTexture, OutputDim, false);
 				}
 				break;
 
@@ -1280,8 +1280,8 @@ void FMediaTextureResource::GetColorSpaceConversionMatrixForSample(const TShared
 					// (this also will deal with any sRGB/Rec703 conversions - as well as any color space conversions)
 					TShaderMapRef<FRGBConvertPS> ConvertShader(ShaderMap);
 					GraphicsPSOInit.BoundShaderState.PixelShaderRHI = ConvertShader.GetPixelShader();
-					SetGraphicsPipelineState(CommandList, GraphicsPSOInit, 0);
-					SetShaderParametersLegacyPS(CommandList, ConvertShader, InputTexture, OutputDim, Encoding, ColorSpaceMtx);
+					SetGraphicsPipelineState(RHICmdList, GraphicsPSOInit, 0);
+					SetShaderParametersLegacyPS(RHICmdList, ConvertShader, InputTexture, OutputDim, Encoding, ColorSpaceMtx);
 				}
 				break;
 
@@ -1291,8 +1291,8 @@ void FMediaTextureResource::GetColorSpaceConversionMatrixForSample(const TShared
 					// (note: no EOTF or CS - this is alpha only!)
 					TShaderMapRef<FRGBConvertPS> ConvertShader(ShaderMap);
 					GraphicsPSOInit.BoundShaderState.PixelShaderRHI = ConvertShader.GetPixelShader();
-					SetGraphicsPipelineState(CommandList, GraphicsPSOInit, 0);
-					SetShaderParametersLegacyPS(CommandList, ConvertShader, InputTexture, OutputDim, UE::Color::EEncoding::Linear, SplatMtx);
+					SetGraphicsPipelineState(RHICmdList, GraphicsPSOInit, 0);
+					SetShaderParametersLegacyPS(RHICmdList, ConvertShader, InputTexture, OutputDim, UE::Color::EEncoding::Linear, SplatMtx);
 				}
 				break;
 
@@ -1317,36 +1317,36 @@ void FMediaTextureResource::GetColorSpaceConversionMatrixForSample(const TShared
 					// (this also will deal with any sRGB/Rec703 conversions - as well as any color space conversions)
 					TShaderMapRef<FYCoCgConvertPS> ConvertShader(ShaderMap);
 					GraphicsPSOInit.BoundShaderState.PixelShaderRHI = ConvertShader.GetPixelShader();
-					SetGraphicsPipelineState(CommandList, GraphicsPSOInit, 0);
-					SetShaderParametersLegacyPS(CommandList, ConvertShader, InputTexture, OutputDim, Encoding, ColorSpaceMtx);
+					SetGraphicsPipelineState(RHICmdList, GraphicsPSOInit, 0);
+					SetShaderParametersLegacyPS(RHICmdList, ConvertShader, InputTexture, OutputDim, Encoding, ColorSpaceMtx);
 				}
 				break;
 
 				default:
 				{
 					// This should not happen in normal use: still - end the render pass to avoid any trouble with RHI
-					CommandList.EndRenderPass();
+					RHICmdList.EndRenderPass();
 					return; // unsupported format (either illegal value or needing a custom converter)
 				}
 			}
 
 			// draw full size quad into render target
 			FBufferRHIRef VertexBuffer = CreateTempMediaVertexBuffer();
-			CommandList.SetStreamSource(0, VertexBuffer, 0);
+			RHICmdList.SetStreamSource(0, VertexBuffer, 0);
 			// set viewport to RT size
-			CommandList.SetViewport(0, 0, 0.0f, (float)OutputDim.X, (float)OutputDim.Y, 1.0f);
+			RHICmdList.SetViewport(0, 0, 0.0f, (float)OutputDim.X, (float)OutputDim.Y, 1.0f);
 
-			CommandList.DrawPrimitive(0, 2, 1);
+			RHICmdList.DrawPrimitive(0, 2, 1);
 		}
-		CommandList.EndRenderPass();
-		CommandList.Transition(FRHITransitionInfo(RenderTarget, ERHIAccess::RTV, ERHIAccess::SRVGraphics));
+		RHICmdList.EndRenderPass();
+		RHICmdList.Transition(FRHITransitionInfo(RenderTarget, ERHIAccess::RTV, ERHIAccess::SRVGraphics));
 	}
 
 	Cleared = false;
 }
 
 
-void FMediaTextureResource::CopySample(const TSharedPtr<IMediaTextureSample, ESPMode::ThreadSafe>& Sample, const FLinearColor& ClearColor, uint8 InNumMips, const FGuid & TextureGUID)
+void FMediaTextureResource::CopySample(FRHICommandListImmediate& RHICmdList, const TSharedPtr<IMediaTextureSample, ESPMode::ThreadSafe>& Sample, const FLinearColor& ClearColor, uint8 InNumMips, const FGuid & TextureGUID)
 {
 	FRHITexture* SampleTexture = Sample->GetTexture();
 	const uint8 SampleNumMips = Sample->GetNumMips();
@@ -1360,7 +1360,7 @@ void FMediaTextureResource::CopySample(const TSharedPtr<IMediaTextureSample, ESP
 		// Use sample's texture as the new render target - no copy
 		if (TextureRHI != SampleTexture)
 		{
-			UpdateTextureReference(SampleTexture);
+			UpdateTextureReference(RHICmdList, SampleTexture);
 
 			MipGenerationCache.SafeRelease();
 			IntermediateTarget.SafeRelease();
@@ -1370,10 +1370,10 @@ void FMediaTextureResource::CopySample(const TSharedPtr<IMediaTextureSample, ESP
 			// Texture to receive texture from sample
 			const uint8 NumMips = (SampleNumMips > 1) ? SampleNumMips : InNumMips;
 			const bool bNeedsUAVTexture = MediaTextureResourceHelpers::RequiresUAVTexture(Sample, NumMips);
-			CreateIntermediateRenderTarget(Sample->GetOutputDim(), MediaTextureResourceHelpers::GetPixelFormat(Sample), MediaTextureResourceHelpers::RequiresSrgbTexture(Sample), ClearColor, NumMips, bNeedsUAVTexture);
+			CreateIntermediateRenderTarget(RHICmdList, Sample->GetOutputDim(), MediaTextureResourceHelpers::GetPixelFormat(Sample), MediaTextureResourceHelpers::RequiresSrgbTexture(Sample), ClearColor, NumMips, bNeedsUAVTexture);
 
 			// Copy data into the output texture to able to add mips later on
-			FRHICommandListExecutor::GetImmediateCommandList().CopyTexture(SampleTexture, IntermediateTarget, FRHICopyTextureInfo());
+			RHICmdList.CopyTexture(SampleTexture, IntermediateTarget, FRHICopyTextureInfo());
 		}
 	}
 	else
@@ -1381,7 +1381,7 @@ void FMediaTextureResource::CopySample(const TSharedPtr<IMediaTextureSample, ESP
 		// Texture to receive precisely only output pixels via CPU copy
 		const uint8 NumMips = (SampleNumMips > 1) ? SampleNumMips : InNumMips;
 		const bool bNeedsUAVTexture = MediaTextureResourceHelpers::RequiresUAVTexture(Sample, NumMips);
-		CreateIntermediateRenderTarget(Sample->GetDim(), MediaTextureResourceHelpers::GetPixelFormat(Sample), MediaTextureResourceHelpers::RequiresSrgbTexture(Sample), ClearColor, NumMips, bNeedsUAVTexture);
+		CreateIntermediateRenderTarget(RHICmdList, Sample->GetDim(), MediaTextureResourceHelpers::GetPixelFormat(Sample), MediaTextureResourceHelpers::RequiresSrgbTexture(Sample), ClearColor, NumMips, bNeedsUAVTexture);
 
 		// If we also have no source buffer and the platform generally would allow for use of external textures, we assume it is just that...
 		// (as long as the player actually produces (dummy) samples, this will enable mips support as well as auto conversion for "new style output" mode)
@@ -1392,19 +1392,18 @@ void FMediaTextureResource::CopySample(const TSharedPtr<IMediaTextureSample, ESP
 
 			if (GSupportsImageExternal)
 			{
-				CopyFromExternalTexture(Sample, TextureGUID);
+				CopyFromExternalTexture(RHICmdList, Sample, TextureGUID);
 			}
 			else
 			{
 				// We never should get here, but could should a player pass us a "valid" sample with neither texture or buffer based data in it (and we don't have ExternalTexture support)
 
 				// Just clear the texture so we don't show any random memory contents...
-				FRHICommandListImmediate& CommandList = FRHICommandListExecutor::GetImmediateCommandList();
-				CommandList.Transition(FRHITransitionInfo(RenderTargetTextureRHI, ERHIAccess::Unknown, ERHIAccess::RTV));
+				RHICmdList.Transition(FRHITransitionInfo(RenderTargetTextureRHI, ERHIAccess::Unknown, ERHIAccess::RTV));
 				FRHIRenderPassInfo RPInfo(RenderTargetTextureRHI, ERenderTargetActions::Clear_Store);
-				CommandList.BeginRenderPass(RPInfo, TEXT("ClearTexture"));
-				CommandList.EndRenderPass();
-				CommandList.Transition(FRHITransitionInfo(RenderTargetTextureRHI, ERHIAccess::RTV, ERHIAccess::SRVMask));
+				RHICmdList.BeginRenderPass(RPInfo, TEXT("ClearTexture"));
+				RHICmdList.EndRenderPass();
+				RHICmdList.Transition(FRHITransitionInfo(RenderTargetTextureRHI, ERHIAccess::RTV, ERHIAccess::SRVMask));
 			}
 		}
 		else
@@ -1422,7 +1421,7 @@ void FMediaTextureResource::CopySample(const TSharedPtr<IMediaTextureSample, ESP
 			}
 
 			// Make sure resource is in SRV mode again
-			FRHICommandListExecutor::GetImmediateCommandList().Transition(FRHITransitionInfo(RenderTargetTextureRHI.GetReference(), ERHIAccess::Unknown, ERHIAccess::SRVMask));
+			RHICmdList.Transition(FRHITransitionInfo(RenderTargetTextureRHI.GetReference(), ERHIAccess::Unknown, ERHIAccess::SRVMask));
 		}
 	}
 
@@ -1430,9 +1429,8 @@ void FMediaTextureResource::CopySample(const TSharedPtr<IMediaTextureSample, ESP
 }
 
 
-void FMediaTextureResource::CopyFromExternalTexture(const TSharedPtr <IMediaTextureSample, ESPMode::ThreadSafe>& Sample, const FGuid & TextureGUID)
+void FMediaTextureResource::CopyFromExternalTexture(FRHICommandListImmediate& RHICmdList, const TSharedPtr <IMediaTextureSample, ESPMode::ThreadSafe>& Sample, const FGuid & TextureGUID)
 {
-	FRHICommandListImmediate& CommandList = FRHICommandListExecutor::GetImmediateCommandList();
 	{
 		FTextureRHIRef SampleTexture;
 		FSamplerStateRHIRef SamplerState;
@@ -1440,9 +1438,9 @@ void FMediaTextureResource::CopyFromExternalTexture(const TSharedPtr <IMediaText
 		{
 			// This should never happen: we could not find the external texture data. Still, if it does we clear the output...
 			FRHIRenderPassInfo RPInfo(RenderTargetTextureRHI, ERenderTargetActions::Clear_Store);
-			CommandList.BeginRenderPass(RPInfo, TEXT("ClearTexture"));
-			CommandList.EndRenderPass();
-			CommandList.Transition(FRHITransitionInfo(RenderTargetTextureRHI, ERHIAccess::Unknown, ERHIAccess::SRVMask));
+			RHICmdList.BeginRenderPass(RPInfo, TEXT("ClearTexture"));
+			RHICmdList.EndRenderPass();
+			RHICmdList.Transition(FRHITransitionInfo(RenderTargetTextureRHI, ERHIAccess::Unknown, ERHIAccess::SRVMask));
 			return;
 		}
 
@@ -1450,19 +1448,19 @@ void FMediaTextureResource::CopyFromExternalTexture(const TSharedPtr <IMediaText
 		FExternalTextureRegistry::Get().GetExternalTextureCoordinateOffset(TextureGUID, Offset);
 		FExternalTextureRegistry::Get().GetExternalTextureCoordinateScaleRotation(TextureGUID, ScaleRotation);
 
-		SCOPED_DRAW_EVENT(CommandList, FMediaTextureResource_ConvertExternalTexture);
-		SCOPED_GPU_STAT(CommandList, MediaTextureResource);
+		SCOPED_DRAW_EVENT(RHICmdList, FMediaTextureResource_ConvertExternalTexture);
+		SCOPED_GPU_STAT(RHICmdList, MediaTextureResource);
 
 		FGraphicsPipelineStateInitializer GraphicsPSOInit;
 		FRHITexture* RenderTarget = RenderTargetTextureRHI.GetReference();
 
 		FRHIRenderPassInfo RPInfo(RenderTarget, ERenderTargetActions::DontLoad_Store);
-		CommandList.BeginRenderPass(RPInfo, TEXT("ConvertMedia_ExternalTexture"));
+		RHICmdList.BeginRenderPass(RPInfo, TEXT("ConvertMedia_ExternalTexture"));
 		{
 			const FIntPoint OutputDim = Sample->GetOutputDim();
 
-			CommandList.ApplyCachedRenderTargets(GraphicsPSOInit);
-			CommandList.SetViewport(0, 0, 0.0f, (float)OutputDim.X, (float)OutputDim.Y, 1.0f);
+			RHICmdList.ApplyCachedRenderTargets(GraphicsPSOInit);
+			RHICmdList.SetViewport(0, 0, 0.0f, (float)OutputDim.X, (float)OutputDim.Y, 1.0f);
 			
 			GraphicsPSOInit.DepthStencilState = TStaticDepthStencilState<false, CF_Always>::GetRHI();
 			GraphicsPSOInit.RasterizerState = TStaticRasterizerState<>::GetRHI();
@@ -1478,19 +1476,19 @@ void FMediaTextureResource::CopyFromExternalTexture(const TSharedPtr <IMediaText
 
 			TShaderMapRef<FReadTextureExternalPS> CopyShader(ShaderMap);
 			GraphicsPSOInit.BoundShaderState.PixelShaderRHI = CopyShader.GetPixelShader();
-			SetGraphicsPipelineState(CommandList, GraphicsPSOInit, 0);
-			SetShaderParametersLegacyPS(CommandList, CopyShader, SampleTexture, SamplerState, ScaleRotation, Offset);
+			SetGraphicsPipelineState(RHICmdList, GraphicsPSOInit, 0);
+			SetShaderParametersLegacyPS(RHICmdList, CopyShader, SampleTexture, SamplerState, ScaleRotation, Offset);
 
 			// draw full size quad into render target
 			FBufferRHIRef VertexBuffer = CreateTempMediaVertexBuffer();
-			CommandList.SetStreamSource(0, VertexBuffer, 0);
+			RHICmdList.SetStreamSource(0, VertexBuffer, 0);
 			// set viewport to RT size
-			CommandList.SetViewport(0, 0, 0.0f, (float)OutputDim.X, (float)OutputDim.Y, 1.0f);
+			RHICmdList.SetViewport(0, 0, 0.0f, (float)OutputDim.X, (float)OutputDim.Y, 1.0f);
 
-			CommandList.DrawPrimitive(0, 2, 1);
+			RHICmdList.DrawPrimitive(0, 2, 1);
 		}
-		CommandList.EndRenderPass();
-		CommandList.Transition(FRHITransitionInfo(RenderTargetTextureRHI, ERHIAccess::Unknown, ERHIAccess::SRVMask));
+		RHICmdList.EndRenderPass();
+		RHICmdList.Transition(FRHITransitionInfo(RenderTargetTextureRHI, ERHIAccess::Unknown, ERHIAccess::SRVMask));
 	}
 }
 
@@ -1518,12 +1516,12 @@ void FMediaTextureResource::UpdateResourceSize()
 }
 
 
-void FMediaTextureResource::UpdateTextureReference(FRHITexture* NewTexture)
+void FMediaTextureResource::UpdateTextureReference(FRHICommandListImmediate& RHICmdList, FRHITexture* NewTexture)
 {
 	TextureRHI = NewTexture;
 	RenderTargetTextureRHI = NewTexture;
 
-	RHIUpdateTextureReference(Owner.TextureReference.TextureReferenceRHI, NewTexture);
+	RHICmdList.UpdateTextureReference(Owner.TextureReference.TextureReferenceRHI, NewTexture);
 	// note: sRGB status for Owner.SRGB is handled (on game thread) in MediaTetxure.cpp
 
 	if (RenderTargetTextureRHI != nullptr)
@@ -1537,7 +1535,7 @@ void FMediaTextureResource::UpdateTextureReference(FRHITexture* NewTexture)
 }
 
 
-void FMediaTextureResource::CreateIntermediateRenderTarget(const FIntPoint & InDim, EPixelFormat InPixelFormat, bool bInSRGB, const FLinearColor & InClearColor, uint8 InNumMips, bool bNeedsUAVSupport)
+void FMediaTextureResource::CreateIntermediateRenderTarget(FRHICommandListImmediate& RHICmdList, const FIntPoint & InDim, EPixelFormat InPixelFormat, bool bInSRGB, const FLinearColor & InClearColor, uint8 InNumMips, bool bNeedsUAVSupport)
 {
 	// create output render target if necessary
 	ETextureCreateFlags OutputCreateFlags = TexCreate_Dynamic | (bInSRGB ? TexCreate_SRGB : TexCreate_None);
@@ -1579,7 +1577,7 @@ void FMediaTextureResource::CreateIntermediateRenderTarget(const FIntPoint & InD
 			.SetClassName(ClassName)
 			.SetOwnerName(GetOwnerName());
 
-		IntermediateTarget = RHICreateTexture(Desc);
+		IntermediateTarget = RHICmdList.CreateTexture(Desc);
 
 		IntermediateTarget->SetName(TEXT("MediaTextureResourceOutput"));
 		IntermediateTarget->SetOwnerName(GetOwnerName());
@@ -1594,7 +1592,7 @@ void FMediaTextureResource::CreateIntermediateRenderTarget(const FIntPoint & InD
 	
 	if (RenderTargetTextureRHI != IntermediateTarget)
 	{
-		UpdateTextureReference(IntermediateTarget);
+		UpdateTextureReference(RHICmdList, IntermediateTarget);
 	}
 }
 
diff --git a/Engine/Source/Runtime/MediaAssets/Private/Misc/MediaTextureResource.h b/Engine/Source/Runtime/MediaAssets/Private/Misc/MediaTextureResource.h
--- a/Engine/Source/Runtime/MediaAssets/Private/Misc/MediaTextureResource.h
+++ b/Engine/Source/Runtime/MediaAssets/Private/Misc/MediaTextureResource.h
@@ -138,7 +138,7 @@ protected:
 	 * @param ClearColor The clear color to use.
 	 * @param SrgbOutput Whether the output texture is in sRGB color space.
 	 */
-	void ClearTexture(const FLinearColor& ClearColor, bool SrgbOutput);
+	void ClearTexture(FRHICommandListImmediate& RHICmdList, const FLinearColor& ClearColor, bool SrgbOutput);
 
 	/**
 	 * Render the given texture sample by converting it on the GPU.
@@ -148,9 +148,9 @@ protected:
 	 * @param Number of mips
 	 * @see CopySample
 	 */
-	void ConvertSample(const TSharedPtr<IMediaTextureSample, ESPMode::ThreadSafe>& Sample, const FLinearColor& ClearColor, uint8 InNumMips);
+	void ConvertSample(FRHICommandListImmediate& RHICmdList, const TSharedPtr<IMediaTextureSample, ESPMode::ThreadSafe>& Sample, const FLinearColor& ClearColor, uint8 InNumMips);
 
-	void ConvertTextureToOutput(FRHITexture* InputTexture, const TSharedPtr<IMediaTextureSample, ESPMode::ThreadSafe>& Sample);
+	void ConvertTextureToOutput(FRHICommandListImmediate& RHICmdList, FRHITexture* InputTexture, const TSharedPtr<IMediaTextureSample, ESPMode::ThreadSafe>& Sample);
 
 	/**
 	 * Render the given texture sample by using it as or copying it to the render target.
@@ -161,7 +161,7 @@ protected:
 	 * @param Number of mips
 	 * @see ConvertSample
 	 */
-	void CopySample(const TSharedPtr<IMediaTextureSample, ESPMode::ThreadSafe>& Sample, const FLinearColor& ClearColor, uint8 InNumMips, const FGuid & TextureGUID);
+	void CopySample(FRHICommandListImmediate& RHICmdList, const TSharedPtr<IMediaTextureSample, ESPMode::ThreadSafe>& Sample, const FLinearColor& ClearColor, uint8 InNumMips, const FGuid & TextureGUID);
 
 	/** Calculates the current resource size and notifies the owner texture. */
 	void UpdateResourceSize();
@@ -171,12 +171,12 @@ protected:
 	 *
 	 * @param NewTexture The texture to set.
 	 */
-	void UpdateTextureReference(FRHITexture* NewTexture);
+	void UpdateTextureReference(FRHICommandListImmediate& RHICmdList, FRHITexture* NewTexture);
 
 	/**
 	 * Create/update intermediate render target as needed. If no color conversion is needed, the RT will be used as the output.
 	 */
-	void CreateIntermediateRenderTarget(const FIntPoint & InDim, EPixelFormat InPixelFormat, bool bInSRGB, const FLinearColor & InClearColor, uint8 InNumMips, bool bNeedsUAVSupport);
+	void CreateIntermediateRenderTarget(FRHICommandListImmediate& RHICmdList, const FIntPoint & InDim, EPixelFormat InPixelFormat, bool bInSRGB, const FLinearColor & InClearColor, uint8 InNumMips, bool bNeedsUAVSupport);
 	
 	/**
 	 * Caches next available sample from queue in MediaTexture owner to keep single consumer access
@@ -189,7 +189,7 @@ protected:
 	void SetupSampler();
 
 	/** Copy to local buffer from external texture */
-	void CopyFromExternalTexture(const TSharedPtr <IMediaTextureSample, ESPMode::ThreadSafe>& Sample, const FGuid & TextureGUID);
+	void CopyFromExternalTexture(FRHICommandListImmediate& RHICmdList, const TSharedPtr <IMediaTextureSample, ESPMode::ThreadSafe>& Sample, const FGuid & TextureGUID);
 
 	bool RequiresConversion(const TSharedPtr<IMediaTextureSample, ESPMode::ThreadSafe>& Sample, uint8 numMips) const;
 	bool RequiresConversion(const FTextureRHIRef& SampleTexture, const FIntPoint & OutputDim, uint8 numMips) const;
