Fixed video flickering issue on Win64 with DX11

#tests Tested videos that were previously flickering
#rnx

diff --git a/Engine/Source/Runtime/MediaAssets/Private/Misc/MediaTextureResource.cpp b/Engine/Source/Runtime/MediaAssets/Private/Misc/MediaTextureResource.cpp
--- a/Engine/Source/Runtime/MediaAssets/Private/Misc/MediaTextureResource.cpp
+++ b/Engine/Source/Runtime/MediaAssets/Private/Misc/MediaTextureResource.cpp
@@ -1050,6 +1050,10 @@ void FMediaTextureResource::GetColorSpaceConversionMatrixForSample(const TShared
 		SCOPED_DRAW_EVENT(RHICmdList, FMediaTextureResource_Convert);
 		SCOPED_GPU_STAT(RHICmdList, MediaTextureResource);
 
+		// draw full size quad into render target
+		// This needs to happen before we begin to setup the draw call, because on DX11, this might flush the command list more or less randomly
+		FBufferRHIRef VertexBuffer = CreateTempMediaVertexBuffer(); 
+
 		FGraphicsPipelineStateInitializer GraphicsPSOInit;
 		FRHITexture* RenderTarget = IntermediateTarget;
 		RHICmdList.Transition(FRHITransitionInfo(RenderTarget, ERHIAccess::Unknown, ERHIAccess::RTV));
@@ -1329,8 +1333,6 @@ void FMediaTextureResource::GetColorSpaceConversionMatrixForSample(const TShared
 				}
 			}
 
-			// draw full size quad into render target
-			FBufferRHIRef VertexBuffer = CreateTempMediaVertexBuffer();
 			RHICmdList.SetStreamSource(0, VertexBuffer, 0);
 			// set viewport to RT size
 			RHICmdList.SetViewport(0, 0, 0.0f, (float)OutputDim.X, (float)OutputDim.Y, 1.0f);
