Add option to disable color conversion on media texture samples

diff --git a/Engine/Plugins/Runtime/nDisplay/Source/SharedMemoryMedia/Private/SharedMemoryMediaSample.h b/Engine/Plugins/Runtime/nDisplay/Source/SharedMemoryMedia/Private/SharedMemoryMediaSample.h
--- a/Engine/Plugins/Runtime/nDisplay/Source/SharedMemoryMedia/Private/SharedMemoryMediaSample.h
+++ b/Engine/Plugins/Runtime/nDisplay/Source/SharedMemoryMedia/Private/SharedMemoryMediaSample.h
@@ -111,6 +111,11 @@ public:
 		return bSrgb;
 	}
 
+	virtual bool ShouldApplyColorConversion() const override
+	{
+		return false;
+	}
+
 public:
 
 #if WITH_ENGINE
diff --git a/Engine/Source/Runtime/Media/Public/IMediaTextureSample.h b/Engine/Source/Runtime/Media/Public/IMediaTextureSample.h
--- a/Engine/Source/Runtime/Media/Public/IMediaTextureSample.h
+++ b/Engine/Source/Runtime/Media/Public/IMediaTextureSample.h
@@ -248,6 +248,12 @@ public:
 	 */
 	virtual uint32 GetStride() const = 0;
 
+	/**
+	 * Whether samples should be converted based on a mismatch with the working color space.
+	 * If this is false, samples will not go through the conversion process even if their encoding or chromaticity doesn't match the working color space.
+	 */
+	virtual bool ShouldApplyColorConversion() const { return true; }
+
 #if WITH_ENGINE
 
 	/**
diff --git a/Engine/Source/Runtime/MediaAssets/Private/Misc/MediaTextureResource.cpp b/Engine/Source/Runtime/MediaAssets/Private/Misc/MediaTextureResource.cpp
--- a/Engine/Source/Runtime/MediaAssets/Private/Misc/MediaTextureResource.cpp
+++ b/Engine/Source/Runtime/MediaAssets/Private/Misc/MediaTextureResource.cpp
@@ -897,11 +897,12 @@ bool FMediaTextureResource::RequiresConversion(const TSharedPtr<IMediaTextureSam
 
 	// Color space different?
 	const UE::Color::FColorSpace& Working = OverrideColorSpace.IsValid() ? *OverrideColorSpace : UE::Color::FColorSpace::GetWorking();
-	const float Tollerance = 1.e-7f;
-	if (!Sample->GetDisplayPrimaryRed().Equals(Working.GetRedChromaticity(), Tollerance) ||
-		!Sample->GetDisplayPrimaryGreen().Equals(Working.GetGreenChromaticity(), Tollerance) ||
-		!Sample->GetDisplayPrimaryBlue().Equals(Working.GetBlueChromaticity(), Tollerance) ||
-		!Sample->GetWhitePoint().Equals(Working.GetWhiteChromaticity(), Tollerance))
+	const float Tolerance = 1.e-7f;
+	if (Sample->ShouldApplyColorConversion() &&
+		(!Sample->GetDisplayPrimaryRed().Equals(Working.GetRedChromaticity(), Tolerance) ||
+		!Sample->GetDisplayPrimaryGreen().Equals(Working.GetGreenChromaticity(), Tolerance) ||
+		!Sample->GetDisplayPrimaryBlue().Equals(Working.GetBlueChromaticity(), Tolerance) ||
+		!Sample->GetWhitePoint().Equals(Working.GetWhiteChromaticity(), Tolerance)))
 	{
 		// Yes! We need to convert...
 		return true;
@@ -938,7 +939,7 @@ bool FMediaTextureResource::RequiresConversion(const TSharedPtr<IMediaTextureSam
 				 Format == EMediaTextureSampleFormat::CharBGR10A2;
 
 	// RGBA and linear?
-	if (bRGBA && ColorEncoding == UE::Color::EEncoding::Linear)
+	if ((bRGBA && ColorEncoding == UE::Color::EEncoding::Linear ) || !Sample->ShouldApplyColorConversion())
 	{
 		return false;
 	}
