MediaAssets: Moved sample queue passing into critical section to avoid potential race condition with UpdatePlayer()

#rnx

diff --git a/Engine/Source/Runtime/MediaAssets/Private/Assets/MediaSoundComponent.cpp b/Engine/Source/Runtime/MediaAssets/Private/Assets/MediaSoundComponent.cpp
--- a/Engine/Source/Runtime/MediaAssets/Private/Assets/MediaSoundComponent.cpp
+++ b/Engine/Source/Runtime/MediaAssets/Private/Assets/MediaSoundComponent.cpp
@@ -485,8 +485,6 @@ ISoundGeneratorPtr UMediaSoundComponent::CreateSoundGenerator(const FSoundGenera
 	FMediaSoundGenerator::FSoundGeneratorParams Params;
 	Params.SampleRate = (int32)InParams.SampleRate;
 	Params.NumChannels = InParams.NumChannels;
-	Params.SampleQueue = SampleQueue;
-	Params.PreviousSampleQueueFlushCount = Params.SampleQueue.IsValid() ? Params.SampleQueue->GetFlushCount() : 0;
 
 	Params.bSpectralAnalysisEnabled = bSpectralAnalysisEnabled;
 	Params.bEnvelopeFollowingEnabled = bEnvelopeFollowingEnabled;
@@ -499,6 +497,9 @@ ISoundGeneratorPtr UMediaSoundComponent::CreateSoundGenerator(const FSoundGenera
 	Params.CachedTime = CachedTime;
 	Params.LastPlaySampleTime = LastPlaySampleTime;
 
+	FScopeLock Lock(&CriticalSection);
+	Params.SampleQueue = SampleQueue;
+	Params.PreviousSampleQueueFlushCount = Params.SampleQueue.IsValid() ? Params.SampleQueue->GetFlushCount() : 0;
 	return MediaSoundGenerator = ISoundGeneratorPtr(new FMediaSoundGenerator(Params));
 }
 
