[Media] MediaTexture: Fix TextureNumMip to reflect the texture resource's number of mips.

diff --git a/Engine/Source/Runtime/MediaAssets/Private/Assets/MediaTexture.cpp b/Engine/Source/Runtime/MediaAssets/Private/Assets/MediaTexture.cpp
--- a/Engine/Source/Runtime/MediaAssets/Private/Assets/MediaTexture.cpp
+++ b/Engine/Source/Runtime/MediaAssets/Private/Assets/MediaTexture.cpp
@@ -175,7 +175,8 @@ FTextureResource* UMediaTexture::CreateResource()
 		Filter = (TextureNumMips > 1) ? TF_Trilinear : TF_Bilinear;
 	}
 
-	return new FMediaTextureResource(*this, Dimensions, Size, ClearColor, CurrentGuid.IsValid() ? CurrentGuid : DefaultGuid, EnableGenMips, NumMips, ColorspaceOverride);
+	const uint8 ClampedNumMips = FMath::Min(NumMips, static_cast<uint8>(GMaxTextureMipCount));
+	return new FMediaTextureResource(*this, Dimensions, Size, ClearColor, CurrentGuid.IsValid() ? CurrentGuid : DefaultGuid, EnableGenMips, ClampedNumMips, ColorspaceOverride);
 }
 
 
@@ -411,7 +412,7 @@ void UMediaTexture::TickResource(FTimespan Timecode)
 				SRGB = Sample->IsOutputSrgb();
 				LastSrgb = SRGB;
 
-				TextureNumMips = (Sample->GetNumMips() > 1) ? Sample->GetNumMips() : NumMips;
+				TextureNumMips = (Sample->GetNumMips() > 1) ? Sample->GetNumMips() : FMediaTextureResource::ClampNumMips(NumMips, Sample->GetDim());
 				bIsSampleValid = true;
 			}
 			else
@@ -431,7 +432,7 @@ void UMediaTexture::TickResource(FTimespan Timecode)
 						SRGB = Sample->IsOutputSrgb();
 						LastSrgb = SRGB;
 
-						TextureNumMips = (Sample->GetNumMips() > 1) ? Sample->GetNumMips() : NumMips;
+						TextureNumMips = (Sample->GetNumMips() > 1) ? Sample->GetNumMips() : FMediaTextureResource::ClampNumMips(NumMips, Sample->GetDim());
 						bIsSampleValid = true;
 					}
 
@@ -479,7 +480,7 @@ void UMediaTexture::TickResource(FTimespan Timecode)
 	RenderParams.ClearColor = ClearColor;
 	RenderParams.PreviousGuid = PreviousGuid;
 	RenderParams.CurrentGuid = CurrentGuid;
-	RenderParams.NumMips = NumMips;
+	RenderParams.NumMips = FMath::Min(NumMips, static_cast<uint8>(GMaxTextureMipCount));
 	
 	// redraw texture resource on render thread
 	FMediaTextureResource* ResourceParam = (FMediaTextureResource*)GetResource();
diff --git a/Engine/Source/Runtime/MediaAssets/Private/Misc/MediaTextureResource.cpp b/Engine/Source/Runtime/MediaAssets/Private/Misc/MediaTextureResource.cpp
--- a/Engine/Source/Runtime/MediaAssets/Private/Misc/MediaTextureResource.cpp
+++ b/Engine/Source/Runtime/MediaAssets/Private/Misc/MediaTextureResource.cpp
@@ -1553,6 +1553,15 @@ void FMediaTextureResource::UpdateTextureReference(FRHICommandListImmediate& RHI
 	}
 }
 
+uint8 FMediaTextureResource::ClampNumMips(uint8 InNumMips, const FIntPoint& InDim)
+{
+	if (InNumMips > 1)
+	{
+		const uint8 MaxMips = static_cast<uint8>(FMath::Min(GMaxTextureMipCount, FGenericPlatformMath::FloorToInt(FGenericPlatformMath::Log2(static_cast<float>(FGenericPlatformMath::Min(InDim.X, InDim.Y))))));
+		InNumMips = FMath::Min(InNumMips, MaxMips);
+	}
+	return InNumMips;
+}
 
 void FMediaTextureResource::CreateIntermediateRenderTarget(FRHICommandListImmediate& RHICmdList, const FIntPoint & InDim, EPixelFormat InPixelFormat, bool bInSRGB, const FLinearColor & InClearColor, uint8 InNumMips, bool bNeedsUAVSupport)
 {
@@ -1563,12 +1572,8 @@ void FMediaTextureResource::CreateIntermediateRenderTarget(FRHICommandListImmedi
 		OutputCreateFlags |= TexCreate_UAV;
 	}
 
-	if (InNumMips > 1)
-	{
-		// Make sure we only set a number of mips that actually makes sense, given the sample size
-		uint8 MaxMips = (uint8)(FMath::Min(255, FGenericPlatformMath::FloorToInt(FGenericPlatformMath::Log2(static_cast<float>(FGenericPlatformMath::Min(InDim.X, InDim.Y))))));
-		InNumMips = FMath::Min(InNumMips, MaxMips);
-	}
+	// Make sure we only set a number of mips that actually makes sense, given the sample size
+	InNumMips = ClampNumMips(InNumMips, InDim);
 
 	// Avoid setting things up as a render target if we are passing in a block compression format
 	// (this happens only if these have no conversion needs at all - otherwise they will be promoted to RGBA8)
diff --git a/Engine/Source/Runtime/MediaAssets/Private/Misc/MediaTextureResource.h b/Engine/Source/Runtime/MediaAssets/Private/Misc/MediaTextureResource.h
--- a/Engine/Source/Runtime/MediaAssets/Private/Misc/MediaTextureResource.h
+++ b/Engine/Source/Runtime/MediaAssets/Private/Misc/MediaTextureResource.h
@@ -114,6 +114,9 @@ public:
 	/** Render the texture using the cached FRenderParams. Call from render thread only. */
 	void JustInTimeRender();
 
+	/** Clamps the number of mips to the maximum needed mips given the texture dimensions. */
+	static uint8 ClampNumMips(uint8 InNumMips, const FIntPoint& InDim);
+	
 public:
 
 	//~ FRenderTarget interface
