Correctly prevent mip generation with block-compressed formats.

We also rename some num mips variables to "InTargetNumMips" for improved clarity.

diff --git a/Engine/Source/Runtime/Media/Public/IMediaTextureSample.h b/Engine/Source/Runtime/Media/Public/IMediaTextureSample.h
--- a/Engine/Source/Runtime/Media/Public/IMediaTextureSample.h
+++ b/Engine/Source/Runtime/Media/Public/IMediaTextureSample.h
@@ -130,6 +130,22 @@ enum class EMediaTextureSampleFormat
 namespace MediaTextureSampleFormat
 {
 	 MEDIA_API const TCHAR* EnumToString(const EMediaTextureSampleFormat InSampleFormat);
+
+	 static inline bool IsBlockCompressedFormat(EMediaTextureSampleFormat InSampleFormat)
+	 {
+		PRAGMA_DISABLE_SWITCH_UNHANDLED_ENUM_CASE_WARNINGS;
+		switch (InSampleFormat)
+		{
+		case EMediaTextureSampleFormat::DXT1:
+		case EMediaTextureSampleFormat::DXT5:
+		case EMediaTextureSampleFormat::BC4:
+		case EMediaTextureSampleFormat::YCoCg_DXT5:
+		case EMediaTextureSampleFormat::YCoCg_DXT5_Alpha_BC4:
+		 return true;
+		}
+		PRAGMA_RESTORE_SWITCH_UNHANDLED_ENUM_CASE_WARNINGS;
+		return false;
+	 }
 };
 
 /** Description of how the media texture sample is tiled (only used by tiled image sequences currently).*/
diff --git a/Engine/Source/Runtime/MediaAssets/Private/Misc/MediaTextureResource.cpp b/Engine/Source/Runtime/MediaAssets/Private/Misc/MediaTextureResource.cpp
--- a/Engine/Source/Runtime/MediaAssets/Private/Misc/MediaTextureResource.cpp
+++ b/Engine/Source/Runtime/MediaAssets/Private/Misc/MediaTextureResource.cpp
@@ -442,8 +442,8 @@ void FMediaTextureResource::Render(const FRenderParams& Params)
 #endif
 		// - - - - - - - - - - - - - - - - - - - - - - - - - - - - - - - - - - - -
 		// - - - - - - - - - - - - - - - - - - - - - - - - - - - - - - - - - - - -
-
-		const uint8 NumMips = bEnableGenMips ? Params.NumMips : 1;
+		const bool bIsBlockCompressedFormat = Sample.IsValid() ? MediaTextureSampleFormat::IsBlockCompressedFormat(Sample->GetFormat()) : false;
+		const uint8 TargetNumMips = (bEnableGenMips && !bIsBlockCompressedFormat) ? Params.NumMips : 1;
 
 		// If real "external texture" support is in place and no mips are used the image will "bypass" any of this processing via the GUID-based lookup for "ExternalTextures" and will
 		// reach the reading material shader without any processing here...
@@ -505,7 +505,7 @@ void FMediaTextureResource::Render(const FRenderParams& Params)
 					if (Converter->Convert(RHICmdList, OutTexture, Hints))
 					{
 						// As the converter created the texture, we might need to convert it even more to make it fit our needs. Check...
-						if (RequiresConversion(OutTexture, Sample->GetOutputDim(), NumMips))
+						if (RequiresConversion(OutTexture, Sample->GetOutputDim(), TargetNumMips))
 						{
 							CreateIntermediateRenderTarget(RHICmdList, Sample->GetOutputDim(), MediaTextureResourceHelpers::GetConvertedPixelFormat(Sample), MediaTextureResourceHelpers::RequiresSrgbTexture(Sample), Params.ClearColor, Hints.NumMips, bNeedsUAVTexture);
 							ConvertTextureToOutput(RHICmdList, OutTexture.GetReference(), Sample);
@@ -543,19 +543,19 @@ void FMediaTextureResource::Render(const FRenderParams& Params)
 			{
 				if (ConvertOrCopyNeeded)
 				{
-					if (RequiresConversion(Sample, NumMips))
+					if (RequiresConversion(Sample, TargetNumMips))
 					{
 						//
 						// Sample needs to be converted by built in converter code
 						//
-						ConvertSample(RHICmdList, Sample, Params.ClearColor, NumMips);
+						ConvertSample(RHICmdList, Sample, Params.ClearColor, TargetNumMips);
 					}
 					else
 					{
 						//
 						// Sample can be used directly or is a simple copy
 						//
-						CopySample(RHICmdList, Sample, Params.ClearColor, NumMips, Params.CurrentGuid);
+						CopySample(RHICmdList, Sample, Params.ClearColor, TargetNumMips, Params.CurrentGuid);
 					}
 				}
 
@@ -622,7 +622,7 @@ void FMediaTextureResource::Render(const FRenderParams& Params)
 				}
 
 				// Generate mips as needed
-				if (CurrentNumMips > 1 && !Cleared && Sample->GetNumMips() == 1)
+				if (CurrentNumMips > 1 && !Cleared && Sample->GetNumMips() == 1 && !IsBlockCompressedFormat(OutputTarget->GetFormat()))
 				{
 					check(OutputTarget);
 
@@ -837,7 +837,7 @@ void FMediaTextureResource::ClearTexture(FRHICommandListImmediate& RHICmdList, c
 }
 
 
-bool FMediaTextureResource::RequiresConversion(const FTextureRHIRef& SampleTexture, const FIntPoint & OutputDim, uint8 InNumMips) const
+bool FMediaTextureResource::RequiresConversion(const FTextureRHIRef& SampleTexture, const FIntPoint & OutputDim, uint8 InTargetNumMips) const
 {
 	if (Owner.NewStyleOutput)
 	{
@@ -847,7 +847,7 @@ bool FMediaTextureResource::RequiresConversion(const FTextureRHIRef& SampleTextu
 
 		// If we have no mips in the sample, but want to have some in the output, we use the conversion pass
 		// to setup level 0 and have a suitable output texture
-		if (SampleTexture->GetNumMips() == 1 && InNumMips != 1)
+		if (SampleTexture->GetNumMips() == 1 && InTargetNumMips != 1)
 		{
 			return true;
 		}
@@ -871,7 +871,7 @@ bool FMediaTextureResource::RequiresConversion(const FTextureRHIRef& SampleTextu
 }
  
 
-bool FMediaTextureResource::RequiresConversion(const TSharedPtr<IMediaTextureSample, ESPMode::ThreadSafe>& Sample, uint8 InNumMips) const
+bool FMediaTextureResource::RequiresConversion(const TSharedPtr<IMediaTextureSample, ESPMode::ThreadSafe>& Sample, uint8 InTargetNumMips) const
 {
 	if (Owner.NewStyleOutput)
 	{
@@ -883,7 +883,7 @@ bool FMediaTextureResource::RequiresConversion(const TSharedPtr<IMediaTextureSam
 
 		// If we have no mips in the sample, but want to have some in the output, we use the conversion pass
 		// to setup level 0 and have a suitable output texture
-		if (Texture && Texture->GetNumMips() == 1 && InNumMips != 1)
+		if (Texture && Texture->GetNumMips() == 1 && InTargetNumMips != 1)
 		{
 			return true;
 		}
@@ -952,7 +952,7 @@ bool FMediaTextureResource::RequiresConversion(const TSharedPtr<IMediaTextureSam
 
 
 
-void FMediaTextureResource::ConvertSample(FRHICommandListImmediate& RHICmdList, const TSharedPtr<IMediaTextureSample, ESPMode::ThreadSafe>& Sample, const FLinearColor& ClearColor, uint8 InNumMips)
+void FMediaTextureResource::ConvertSample(FRHICommandListImmediate& RHICmdList, const TSharedPtr<IMediaTextureSample, ESPMode::ThreadSafe>& Sample, const FLinearColor& ClearColor, uint8 InTargetNumMips)
 {
 	const EPixelFormat InputPixelFormat = MediaTextureResourceHelpers::GetPixelFormat(Sample);
 
@@ -1014,7 +1014,7 @@ void FMediaTextureResource::ConvertSample(FRHICommandListImmediate& RHICmdList,
 
 	// create the output texture
 	const FIntPoint OutputDim = Sample->GetOutputDim();
-	const uint8 NumMips = (SampleNumMips > 1) ? SampleNumMips : InNumMips;
+	const uint8 NumMips = (SampleNumMips > 1) ? SampleNumMips : InTargetNumMips;
 	const bool bNeedsUAVTexture = MediaTextureResourceHelpers::RequiresUAVTexture(Sample, NumMips);
 	CreateIntermediateRenderTarget(RHICmdList, OutputDim, MediaTextureResourceHelpers::GetConvertedPixelFormat(Sample), MediaTextureResourceHelpers::RequiresSrgbTexture(Sample), ClearColor, NumMips, bNeedsUAVTexture);
 
@@ -1365,7 +1365,7 @@ void FMediaTextureResource::GetColorSpaceConversionMatrixForSample(const TShared
 }
 
 
-void FMediaTextureResource::CopySample(FRHICommandListImmediate& RHICmdList, const TSharedPtr<IMediaTextureSample, ESPMode::ThreadSafe>& Sample, const FLinearColor& ClearColor, uint8 InNumMips, const FGuid & TextureGUID)
+void FMediaTextureResource::CopySample(FRHICommandListImmediate& RHICmdList, const TSharedPtr<IMediaTextureSample, ESPMode::ThreadSafe>& Sample, const FLinearColor& ClearColor, uint8 InTargetNumMips, const FGuid & TextureGUID)
 {
 	FRHITexture* SampleTexture = Sample->GetTexture();
 	const uint8 SampleNumMips = Sample->GetNumMips();
@@ -1387,7 +1387,7 @@ void FMediaTextureResource::CopySample(FRHICommandListImmediate& RHICmdList, con
 		else
 		{
 			// Texture to receive texture from sample
-			const uint8 NumMips = (SampleNumMips > 1) ? SampleNumMips : InNumMips;
+			const uint8 NumMips = (SampleNumMips > 1) ? SampleNumMips : InTargetNumMips;
 			const bool bNeedsUAVTexture = MediaTextureResourceHelpers::RequiresUAVTexture(Sample, NumMips);
 			CreateIntermediateRenderTarget(RHICmdList, Sample->GetOutputDim(), MediaTextureResourceHelpers::GetPixelFormat(Sample), MediaTextureResourceHelpers::RequiresSrgbTexture(Sample), ClearColor, NumMips, bNeedsUAVTexture);
 
@@ -1398,7 +1398,7 @@ void FMediaTextureResource::CopySample(FRHICommandListImmediate& RHICmdList, con
 	else
 	{
 		// Texture to receive precisely only output pixels via CPU copy
-		const uint8 NumMips = (SampleNumMips > 1) ? SampleNumMips : InNumMips;
+		const uint8 NumMips = (SampleNumMips > 1) ? SampleNumMips : InTargetNumMips;
 		const bool bNeedsUAVTexture = MediaTextureResourceHelpers::RequiresUAVTexture(Sample, NumMips);
 		CreateIntermediateRenderTarget(RHICmdList, Sample->GetDim(), MediaTextureResourceHelpers::GetPixelFormat(Sample), MediaTextureResourceHelpers::RequiresSrgbTexture(Sample), ClearColor, NumMips, bNeedsUAVTexture);
 
@@ -1577,7 +1577,7 @@ void FMediaTextureResource::CreateIntermediateRenderTarget(FRHICommandListImmedi
 
 	// Avoid setting things up as a render target if we are passing in a block compression format
 	// (this happens only if these have no conversion needs at all - otherwise they will be promoted to RGBA8)
-	if (InPixelFormat != PF_DXT1 && InPixelFormat != PF_DXT5 && InPixelFormat != PF_BC4)
+	if (!IsBlockCompressedFormat(InPixelFormat))
 	{
 		OutputCreateFlags |= ETextureCreateFlags::RenderTargetable;
 	}
diff --git a/Engine/Source/Runtime/MediaAssets/Private/Misc/MediaTextureResource.h b/Engine/Source/Runtime/MediaAssets/Private/Misc/MediaTextureResource.h
--- a/Engine/Source/Runtime/MediaAssets/Private/Misc/MediaTextureResource.h
+++ b/Engine/Source/Runtime/MediaAssets/Private/Misc/MediaTextureResource.h
@@ -151,7 +151,7 @@ protected:
 	 * @param Number of mips
 	 * @see CopySample
 	 */
-	void ConvertSample(FRHICommandListImmediate& RHICmdList, const TSharedPtr<IMediaTextureSample, ESPMode::ThreadSafe>& Sample, const FLinearColor& ClearColor, uint8 InNumMips);
+	void ConvertSample(FRHICommandListImmediate& RHICmdList, const TSharedPtr<IMediaTextureSample, ESPMode::ThreadSafe>& Sample, const FLinearColor& ClearColor, uint8 InTargetNumMips);
 
 	void ConvertTextureToOutput(FRHICommandListImmediate& RHICmdList, FRHITexture* InputTexture, const TSharedPtr<IMediaTextureSample, ESPMode::ThreadSafe>& Sample);
 
@@ -164,7 +164,7 @@ protected:
 	 * @param Number of mips
 	 * @see ConvertSample
 	 */
-	void CopySample(FRHICommandListImmediate& RHICmdList, const TSharedPtr<IMediaTextureSample, ESPMode::ThreadSafe>& Sample, const FLinearColor& ClearColor, uint8 InNumMips, const FGuid & TextureGUID);
+	void CopySample(FRHICommandListImmediate& RHICmdList, const TSharedPtr<IMediaTextureSample, ESPMode::ThreadSafe>& Sample, const FLinearColor& ClearColor, uint8 InTargetNumMips, const FGuid & TextureGUID);
 
 	/** Calculates the current resource size and notifies the owner texture. */
 	void UpdateResourceSize();
@@ -194,8 +194,8 @@ protected:
 	/** Copy to local buffer from external texture */
 	void CopyFromExternalTexture(FRHICommandListImmediate& RHICmdList, const TSharedPtr <IMediaTextureSample, ESPMode::ThreadSafe>& Sample, const FGuid & TextureGUID);
 
-	bool RequiresConversion(const TSharedPtr<IMediaTextureSample, ESPMode::ThreadSafe>& Sample, uint8 numMips) const;
-	bool RequiresConversion(const FTextureRHIRef& SampleTexture, const FIntPoint & OutputDim, uint8 numMips) const;
+	bool RequiresConversion(const TSharedPtr<IMediaTextureSample, ESPMode::ThreadSafe>& Sample, uint8 InTargetNumMips) const;
+	bool RequiresConversion(const FTextureRHIRef& SampleTexture, const FIntPoint & OutputDim, uint8 InTargetNumMips) const;
 
 	/** Compute CS conversion martix based on sample's data */
 	void GetColorSpaceConversionMatrixForSample(const TSharedPtr<IMediaTextureSample, ESPMode::ThreadSafe> Sample, FMatrix44f& ColorSpaceMtx);
