[MediaAssets] Fix the expected number of mips in MediaTexture, taking into account if mipmap generation is enable or not.

diff --git a/Engine/Source/Runtime/MediaAssets/Private/Assets/MediaTexture.cpp b/Engine/Source/Runtime/MediaAssets/Private/Assets/MediaTexture.cpp
--- a/Engine/Source/Runtime/MediaAssets/Private/Assets/MediaTexture.cpp
+++ b/Engine/Source/Runtime/MediaAssets/Private/Assets/MediaTexture.cpp
@@ -351,6 +351,34 @@ void UMediaTexture::PostEditChangeProperty(FPropertyChangedEvent& PropertyChange
 /* UMediaTexture implementation
  *****************************************************************************/
 
+namespace UE::MediaTexture::Private
+{
+	/**
+	 * @brief Calculates the expected number of mips for the given sample that the texture resource should have.
+	 * @param InSample Decoded sample, can't be null.
+	 * @param InRequestedNumMips Requested number of mips from the media texture.
+	 * @param bInEnableGenMips Indicate if mip generation is enabled.
+	 * @return Expected number of mips.
+	 */
+	int32 CalcTextureNumMips(const TSharedPtr<IMediaTextureSample, ESPMode::ThreadSafe>& InSample, int8 InRequestedNumMips, bool bInEnableGenMips)
+	{
+		check(InSample);
+		
+		if (InSample->GetNumMips() > 1)
+		{
+			return InSample->GetNumMips();
+		}
+
+		if (bInEnableGenMips)
+		{
+			return FMediaTextureResource::ClampNumMips(InRequestedNumMips, InSample->GetDim());
+		}
+
+		// If mipmap generation is not enabled, ignore the requested number of mips to reflect texture resource's state.
+		return 1;
+	}
+}
+
 void UMediaTexture::TickResource(FTimespan Timecode)
 {
 	if (GetResource() == nullptr)
@@ -412,7 +440,7 @@ void UMediaTexture::TickResource(FTimespan Timecode)
 				SRGB = Sample->IsOutputSrgb();
 				LastSrgb = SRGB;
 
-				TextureNumMips = (Sample->GetNumMips() > 1) ? Sample->GetNumMips() : FMediaTextureResource::ClampNumMips(NumMips, Sample->GetDim());
+				TextureNumMips = UE::MediaTexture::Private::CalcTextureNumMips(Sample, NumMips, EnableGenMips);
 				bIsSampleValid = true;
 			}
 			else
@@ -432,7 +460,7 @@ void UMediaTexture::TickResource(FTimespan Timecode)
 						SRGB = Sample->IsOutputSrgb();
 						LastSrgb = SRGB;
 
-						TextureNumMips = (Sample->GetNumMips() > 1) ? Sample->GetNumMips() : FMediaTextureResource::ClampNumMips(NumMips, Sample->GetDim());
+						TextureNumMips = UE::MediaTexture::Private::CalcTextureNumMips(Sample, NumMips, EnableGenMips);
 						bIsSampleValid = true;
 					}
 
