Media: Minor cleanup of num mips property, as a follow-up to recent deprecations.

We also avoid the UAV resource requirement unless generate mips is actually enabled.

diff --git a/Engine/Source/Runtime/MediaAssets/Private/Assets/MediaTexture.cpp b/Engine/Source/Runtime/MediaAssets/Private/Assets/MediaTexture.cpp
--- a/Engine/Source/Runtime/MediaAssets/Private/Assets/MediaTexture.cpp
+++ b/Engine/Source/Runtime/MediaAssets/Private/Assets/MediaTexture.cpp
@@ -353,36 +353,10 @@ void UMediaTexture::PostEditChangeProperty(FPropertyChangedEvent& PropertyChange
 /* UMediaTexture implementation
  *****************************************************************************/
 
-namespace UE::MediaTexture::Private
-{
-	/**
-	 * @brief Calculates the expected number of mips for the given sample that the texture resource should have.
-	 * @param InSample Decoded sample, can't be null.
-	 * @param InRequestedNumMips Requested number of mips from the media texture.
-	 * @param bInEnableGenMips Indicate if mip generation is enabled.
-	 * @return Expected number of mips.
-	 */
-	int32 CalcTextureNumMips(const TSharedPtr<IMediaTextureSample, ESPMode::ThreadSafe>& InSample, int8 InRequestedNumMips, bool bInEnableGenMips)
-	{
-		check(InSample);
-		
-		if (InSample->GetNumMips() > 1)
-		{
-			return InSample->GetNumMips();
-		}
-
-		if (bInEnableGenMips)
-		{
-			return FMediaTextureResource::ClampNumMips(InRequestedNumMips, InSample->GetDim());
-		}
-
-		// If mipmap generation is not enabled, ignore the requested number of mips to reflect texture resource's state.
-		return 1;
-	}
-}
-
 void UMediaTexture::TickResource(FTimespan Timecode)
 {
+	using namespace UE::MediaTexture::Private;
+
 	if (GetResource() == nullptr)
 	{
 		return;
@@ -403,7 +377,6 @@ void UMediaTexture::TickResource(FTimespan Timecode)
 
 	LastClearColor = ClearColor;
 	LastSrgb = SRGB;
-	const uint8 MaxNumMips = static_cast<uint8>(GMaxTextureMipCount);
 
 	// set up render parameters
 	FMediaTextureResource::FRenderParams RenderParams;
@@ -443,7 +416,7 @@ void UMediaTexture::TickResource(FTimespan Timecode)
 				SRGB = Sample->IsOutputSrgb();
 				LastSrgb = SRGB;
 
-				TextureNumMips = UE::MediaTexture::Private::CalcTextureNumMips(Sample, MaxNumMips, EnableGenMips);
+				TextureNumMips = static_cast<int32>(CalcTextureNumMips(*Sample, EnableGenMips));
 				bIsSampleValid = true;
 			}
 			else
@@ -463,7 +436,7 @@ void UMediaTexture::TickResource(FTimespan Timecode)
 						SRGB = Sample->IsOutputSrgb();
 						LastSrgb = SRGB;
 
-						TextureNumMips = UE::MediaTexture::Private::CalcTextureNumMips(Sample, MaxNumMips, EnableGenMips);
+						TextureNumMips = static_cast<int32>(CalcTextureNumMips(*Sample, EnableGenMips));
 						bIsSampleValid = true;
 					}
 
@@ -511,7 +484,6 @@ void UMediaTexture::TickResource(FTimespan Timecode)
 	RenderParams.ClearColor = ClearColor;
 	RenderParams.PreviousGuid = PreviousGuid;
 	RenderParams.CurrentGuid = CurrentGuid;
-	RenderParams.NumMips = MaxNumMips;
 
 	// redraw texture resource on render thread
 	FMediaTextureResource* ResourceParam = (FMediaTextureResource*)GetResource();
diff --git a/Engine/Source/Runtime/MediaAssets/Private/Misc/MediaTextureResource.cpp b/Engine/Source/Runtime/MediaAssets/Private/Misc/MediaTextureResource.cpp
--- a/Engine/Source/Runtime/MediaAssets/Private/Misc/MediaTextureResource.cpp
+++ b/Engine/Source/Runtime/MediaAssets/Private/Misc/MediaTextureResource.cpp
@@ -48,6 +48,37 @@ DECLARE_GPU_STAT_NAMED(MediaTextureResource, TEXT("MediaTextureResource"));
 
 static const FMatrix44f SplatMtx = { {1.0f, 0.0f, 0.0f, 0.0f}, {1.0f, 0.0f, 0.0f, 0.0f}, {1.0f, 0.0f, 0.0f, 0.0f}, {1.0f, 0.0f, 0.0f, 0.0f} };
 
+namespace UE::MediaTexture::Private
+{
+	uint8 ClampNumMips(uint8 InNumMips, const FIntPoint& InDim)
+	{
+		if (InNumMips > 1u)
+		{
+			const uint8 MaxNumMips = static_cast<uint8>(FMath::Min(GMaxTextureMipCount, FGenericPlatformMath::FloorToInt(FGenericPlatformMath::Log2(static_cast<float>(FGenericPlatformMath::Min(InDim.X, InDim.Y))))));
+			InNumMips = FMath::Min(InNumMips, MaxNumMips);
+		}
+
+		return InNumMips;
+	}
+
+	uint8 CalcGeneratedNumMips(const IMediaTextureSample& InSample, bool bInEnableGenMips)
+	{
+		if (bInEnableGenMips && !MediaTextureSampleFormat::IsBlockCompressedFormat(InSample.GetFormat()))
+		{
+			return ClampNumMips(static_cast<uint8>(GMaxTextureMipCount), InSample.GetDim());
+		}
+
+		return 1u;
+	}
+
+	uint8 CalcTextureNumMips(const IMediaTextureSample& InSample, bool bInEnableGenMips)
+	{
+		const uint8 SampleNumMips = InSample.GetNumMips();
+
+		return SampleNumMips > 1u ? SampleNumMips : CalcGeneratedNumMips(InSample, bInEnableGenMips);
+	}
+}
+
 /* GPU data deletion helper
  *****************************************************************************/
 
@@ -440,10 +471,6 @@ void FMediaTextureResource::Render(const FRenderParams& Params)
 #if UE_MEDIAUTILS_DEVELOPMENT_DELEGATE
 		FMediaDelegates::OnPreSampleRender_RenderThread.Broadcast(&Owner, UseSample, Sample);
 #endif
-		// - - - - - - - - - - - - - - - - - - - - - - - - - - - - - - - - - - - -
-		// - - - - - - - - - - - - - - - - - - - - - - - - - - - - - - - - - - - -
-		const bool bIsBlockCompressedFormat = Sample.IsValid() ? MediaTextureSampleFormat::IsBlockCompressedFormat(Sample->GetFormat()) : false;
-		const uint8 TargetNumMips = (bEnableGenMips && !bIsBlockCompressedFormat) ? Params.NumMips : 1;
 
 		// If real "external texture" support is in place and no mips are used the image will "bypass" any of this processing via the GUID-based lookup for "ExternalTextures" and will
 		// reach the reading material shader without any processing here...
@@ -454,6 +481,8 @@ void FMediaTextureResource::Render(const FRenderParams& Params)
 			// Valid sample & Sample should be shown
 			//
 
+			const uint8 TargetNumMips = UE::MediaTexture::Private::CalcGeneratedNumMips(*Sample, bEnableGenMips);
+
 			bool ConvertOrCopyNeeded = false;
 
 			if (Sample->GetOutputDim().GetMin() <= 0)
@@ -472,7 +501,7 @@ void FMediaTextureResource::Render(const FRenderParams& Params)
 				const uint8 SampleNumMips = Sample->GetNumMips();
 
 				IMediaTextureSampleConverter::FConversionHints Hints;
-				Hints.NumMips = (SampleNumMips > 1) ? SampleNumMips : Params.NumMips;
+				Hints.NumMips = (SampleNumMips > 1) ? SampleNumMips : TargetNumMips;
 
 				const bool bNeedsUAVTexture = MediaTextureResourceHelpers::RequiresUAVTexture(Sample, Hints.NumMips);
 
@@ -1583,16 +1612,6 @@ void FMediaTextureResource::UpdateTextureReference(FRHICommandListImmediate& RHI
 	}
 }
 
-uint8 FMediaTextureResource::ClampNumMips(uint8 InNumMips, const FIntPoint& InDim)
-{
-	if (InNumMips > 1)
-	{
-		const uint8 MaxMips = static_cast<uint8>(FMath::Min(GMaxTextureMipCount, FGenericPlatformMath::FloorToInt(FGenericPlatformMath::Log2(static_cast<float>(FGenericPlatformMath::Min(InDim.X, InDim.Y))))));
-		InNumMips = FMath::Min(InNumMips, MaxMips);
-	}
-	return InNumMips;
-}
-
 void FMediaTextureResource::CreateIntermediateRenderTarget(FRHICommandListImmediate& RHICmdList, const FIntPoint & InDim, EPixelFormat InPixelFormat, bool bInSRGB, const FLinearColor & InClearColor, uint8 InNumMips, bool bNeedsUAVSupport)
 {
 	// create output render target if necessary
@@ -1603,7 +1622,7 @@ void FMediaTextureResource::CreateIntermediateRenderTarget(FRHICommandListImmedi
 	}
 
 	// Make sure we only set a number of mips that actually makes sense, given the sample size
-	InNumMips = ClampNumMips(InNumMips, InDim);
+	InNumMips = UE::MediaTexture::Private::ClampNumMips(InNumMips, InDim);
 
 	// Avoid setting things up as a render target if we are passing in a block compression format
 	// (this happens only if these have no conversion needs at all - otherwise they will be promoted to RGBA8)
diff --git a/Engine/Source/Runtime/MediaAssets/Private/Misc/MediaTextureResource.h b/Engine/Source/Runtime/MediaAssets/Private/Misc/MediaTextureResource.h
--- a/Engine/Source/Runtime/MediaAssets/Private/Misc/MediaTextureResource.h
+++ b/Engine/Source/Runtime/MediaAssets/Private/Misc/MediaTextureResource.h
@@ -29,6 +29,31 @@ struct FPriorSamples;
 enum class EMediaTextureSinkFormat;
 enum class EMediaTextureSinkMode;
 
+namespace UE::MediaTexture::Private
+{
+	/** Clamps the number of mips to the maximum needed mips given the texture dimensions. */
+	uint8 ClampNumMips(uint8 InNumMips, const FIntPoint& InDim);
+
+	/**
+	 * @brief Calculates the target number of generated mips for a given sample.
+	 * 
+	 * @param InSample Decoded sample.
+	 * @param bInEnableGenMips Indicate if mip generation is enabled.
+	 * @return Expected number of generated mips.
+	 */
+	uint8 CalcGeneratedNumMips(const IMediaTextureSample& InSample, bool bInEnableGenMips);
+
+	/**
+	 * @brief Calculates the expected texture resource number of mips for the given sample.
+	 * Sample source mips take priority over generated ones.
+	 *
+	 * @param InSample Decoded sample.
+	 * @param bInEnableGenMips Indicate if mip generation is enabled.
+	 * @return Expected number of mips.
+	 */
+	uint8 CalcTextureNumMips(const IMediaTextureSample& InSample, bool bInEnableGenMips);
+}
+
 /**
  * Texture resource type for media textures.
  */
@@ -61,6 +86,14 @@ public:
 	/** Parameters for the Render method. */
 	struct FRenderParams
 	{
+		PRAGMA_DISABLE_DEPRECATION_WARNINGS
+		FRenderParams() = default;
+		FRenderParams(const FRenderParams& InOther) = default;
+		FRenderParams(FRenderParams&& InOther) = default;
+		FRenderParams& operator=(const FRenderParams&) = default;
+		FRenderParams& operator=(FRenderParams&&) = default;
+		PRAGMA_ENABLE_DEPRECATION_WARNINGS
+
 		/** Whether the texture can be cleared. */
 		bool CanClear;
 
@@ -80,6 +113,7 @@ public:
 		TWeakPtr<FMediaTextureSampleSource, ESPMode::ThreadSafe> SampleSource;
 
 		/** Number of mips wanted */
+		UE_DEPRECATED(5.6, "Unused since UMediaTexture NumMips is now deprecated.")
 		uint8 NumMips;
 
 		/** The time of the video frame to render (in player's clock). */
@@ -113,9 +147,6 @@ public:
 
 	/** Render the texture using the cached FRenderParams. Call from render thread only. */
 	void JustInTimeRender();
-
-	/** Clamps the number of mips to the maximum needed mips given the texture dimensions. */
-	static uint8 ClampNumMips(uint8 InNumMips, const FIntPoint& InDim);
 	
 public:
 
