Media: Updated new style output tooltip, deprecation cleanup.

diff --git a/Engine/Source/Runtime/MediaAssets/Public/MediaTexture.h b/Engine/Source/Runtime/MediaAssets/Public/MediaTexture.h
--- a/Engine/Source/Runtime/MediaAssets/Public/MediaTexture.h
+++ b/Engine/Source/Runtime/MediaAssets/Public/MediaTexture.h
@@ -21,14 +21,6 @@ class FMediaTextureClockSink;
 class IMediaTextureSample;
 class UMediaPlayer;
 
-UENUM()
-enum UE_DEPRECATED(5.4, "This enum was unused (not connected to active logic) and is now deprecated.") MediaTextureOutputFormat : int
-{
-	MTOF_Default					UMETA(DisplayName = "Default (sRGB)"),
-	MTOF_SRGB_LINOUT				UMETA(DisplayName = "sRGB (linear output)"),		// sRGB data, using sRGB texture formats; hence read as linear RGB
-	MTOF_MAX,
-};
-
 UENUM()
 enum MediaTextureOrientation : int
 {
@@ -87,18 +79,13 @@ public:
 	uint8 NumMips_DEPRECATED;
 #endif
 
-	/** Enable new style output. */
-	UPROPERTY(EditAnywhere, BlueprintReadWrite, Category = "MediaTexture", meta = (DisplayName = "Enable new style output"))
+	/**
+	* Allows the media texture to be used as a regular 2D texture in materials, with Color or Linear Color sampler types as opposed to External.
+	* Note however that external textures can remain more efficient on lower-end platforms (OpenGL ES).
+	*/
+	UPROPERTY(EditAnywhere, BlueprintReadWrite, Category = "MediaTexture", meta = (DisplayName = "Enable New Style Output"))
 	bool NewStyleOutput;
 
-PRAGMA_DISABLE_DEPRECATION_WARNINGS
-#if WITH_EDITORONLY_DATA
-	/** DEPRECATED 5.4 */
-	UPROPERTY(meta = (DeprecatedProperty, DeprecationMessage = "Output format was unused (not connected to active logic) and is now deprecated. References to it can be safely deleted."))
-	TEnumAsByte<enum MediaTextureOutputFormat> OutputFormat_DEPRECATED;
-#endif
-PRAGMA_ENABLE_DEPRECATION_WARNINGS
-
 	/** Current aspect ratio */
 	UPROPERTY(Transient, TextExportTransient, SkipSerialization, BlueprintReadOnly, Category = "MediaTexture", meta = (DisplayName = "Current frame's aspect ratio"))
 	float CurrentAspectRatio;
