Use a transactionally-safe critical section in `UMediaSoundComponent`.

diff --git a/Engine/Source/Runtime/MediaAssets/Private/Assets/MediaSoundComponent.cpp b/Engine/Source/Runtime/MediaAssets/Private/Assets/MediaSoundComponent.cpp
--- a/Engine/Source/Runtime/MediaAssets/Private/Assets/MediaSoundComponent.cpp
+++ b/Engine/Source/Runtime/MediaAssets/Private/Assets/MediaSoundComponent.cpp
@@ -267,7 +267,7 @@ void UMediaSoundComponent::UpdatePlayer()
 		CachedRate = 0.0f;
 		CachedTime = FTimespan::Zero();
 
-		FScopeLock Lock(&CriticalSection);
+		UE::TScopeLock Lock(CriticalSection);
 		SampleQueue.Reset();
 		MediaSoundGenerator.Reset();
 		return;
@@ -289,7 +289,7 @@ void UMediaSoundComponent::UpdatePlayer()
 			const auto NewSampleQueue = MakeShared<FMediaAudioSampleQueue, ESPMode::ThreadSafe>(MaxAudioInputSamples);
 			PlayerFacade->AddAudioSampleSink(NewSampleQueue);
 			{
-				FScopeLock Lock(&CriticalSection);
+				UE::TScopeLock Lock(CriticalSection);
 				SampleQueue = NewSampleQueue;
 				if (MediaSoundGenerator.IsValid())
 				{
@@ -305,7 +305,7 @@ void UMediaSoundComponent::UpdatePlayer()
 		// We need to check for deactivation as it seems there is not callback scheduled when USynthComponent::Stop() is called.
 		if(!IsActive())
 		{
-			FScopeLock Lock(&CriticalSection);
+			UE::TScopeLock Lock(CriticalSection);
 			SampleQueue.Reset();
 			CurrentPlayerFacade.Reset();
 		}
@@ -374,7 +374,7 @@ void UMediaSoundComponent::OnRegister()
 void UMediaSoundComponent::OnUnregister()
 {
 	{
-		FScopeLock Lock(&CriticalSection);
+		UE::TScopeLock Lock(CriticalSection);
 		SampleQueue.Reset();
 	}
 	CurrentPlayerFacade.Reset();
@@ -411,7 +411,7 @@ void UMediaSoundComponent::Deactivate()
 	{
 		SetComponentTickEnabled(false);
 		{
-			FScopeLock Lock(&CriticalSection);
+			UE::TScopeLock Lock(CriticalSection);
 			SampleQueue.Reset();
 		}
 		CurrentPlayerFacade.Reset();
@@ -497,7 +497,7 @@ ISoundGeneratorPtr UMediaSoundComponent::CreateSoundGenerator(const FSoundGenera
 	Params.CachedTime = CachedTime;
 	Params.LastPlaySampleTime = LastPlaySampleTime;
 
-	FScopeLock Lock(&CriticalSection);
+	UE::TScopeLock Lock(CriticalSection);
 	Params.SampleQueue = SampleQueue;
 	Params.PreviousSampleQueueFlushCount = Params.SampleQueue.IsValid() ? Params.SampleQueue->GetFlushCount() : 0;
 	return MediaSoundGenerator = ISoundGeneratorPtr(new FMediaSoundGenerator(Params));
diff --git a/Engine/Source/Runtime/MediaAssets/Public/MediaSoundComponent.h b/Engine/Source/Runtime/MediaAssets/Public/MediaSoundComponent.h
--- a/Engine/Source/Runtime/MediaAssets/Public/MediaSoundComponent.h
+++ b/Engine/Source/Runtime/MediaAssets/Public/MediaSoundComponent.h
@@ -5,7 +5,6 @@
 #include "Components/SynthComponent.h"
 
 #include "Containers/Array.h"
-#include "HAL/CriticalSection.h"
 #include "MediaSampleQueue.h"
 #include "Misc/Timespan.h"
 #include "Templates/Atomic.h"
@@ -18,6 +17,7 @@
 #include "Sound/SoundClass.h"
 #include "Sound/SoundGenerator.h"
 #include "MediaAudioResampler.h"
+#include "Misc/TransactionallySafeCriticalSection.h"
 
 #include "MediaSoundComponent.generated.h"
 
@@ -364,7 +364,7 @@ private:
 	FTimespan CachedTime;
 
 	/** Critical section for synchronizing access to PlayerFacadePtr. */
-	FCriticalSection CriticalSection;
+	FTransactionallySafeCriticalSection CriticalSection;
 
 	/** The player that is currently associated with this component. */
 	TWeakObjectPtr<UMediaPlayer> CurrentPlayer;
