[Electra] Avid DNxHD support 12 bits (HQX) format.

diff --git a/Engine/Plugins/Media/AvidDNxHDDecoderElectra/Source/Private/Windows/ElectraMediaAvidDNxHDDecoder.cpp b/Engine/Plugins/Media/AvidDNxHDDecoderElectra/Source/Private/Windows/ElectraMediaAvidDNxHDDecoder.cpp
--- a/Engine/Plugins/Media/AvidDNxHDDecoderElectra/Source/Private/Windows/ElectraMediaAvidDNxHDDecoder.cpp
+++ b/Engine/Plugins/Media/AvidDNxHDDecoderElectra/Source/Private/Windows/ElectraMediaAvidDNxHDDecoder.cpp
@@ -962,6 +962,18 @@ bool FVideoDecoderAvidDNxHDElectra::FDecoderHandle::SetupDecompressStruct()
 				CurrentUncompressedParams.compOrder = DecodeOperationParams.decodeAlpha ? DNX_ColorComponentOrder_t::DNX_CCO_CbYCrY_A : DNX_ColorComponentOrder_t::DNX_CCO_YCbYCr_NoA;
 				break;
 			}
+		case 12:
+				{
+					CurrentUncompressedParams.compType = DNX_ComponentType_t::DNX_CT_USHORT_12_4;
+					// FIXME: What output format would the alpha channel require us to use?
+					if (DecodeOperationParams.decodeAlpha)
+					{
+						UE_LOG(LogAvidDNxHDElectraDecoder, Log, TEXT("Compressed alpha with YCbYCr is not currently handled"));
+						return false;
+					}
+					CurrentUncompressedParams.compOrder = DecodeOperationParams.decodeAlpha ? DNX_ColorComponentOrder_t::DNX_CCO_CbYCrY_A : DNX_ColorComponentOrder_t::DNX_CCO_YCbYCr_NoA;
+					break;
+				}
 			default:
 			{
 				UE_LOG(LogAvidDNxHDElectraDecoder, Log, TEXT("Compressed %d bit depth is not currently handled"), CurrentCompressedParams.depth);
diff --git a/Engine/Plugins/Media/ElectraUtil/Source/ElectraSamples/Private/Windows/ElectraTextureSample_Windows.cpp b/Engine/Plugins/Media/ElectraUtil/Source/ElectraSamples/Private/Windows/ElectraTextureSample_Windows.cpp
--- a/Engine/Plugins/Media/ElectraUtil/Source/ElectraSamples/Private/Windows/ElectraTextureSample_Windows.cpp
+++ b/Engine/Plugins/Media/ElectraUtil/Source/ElectraSamples/Private/Windows/ElectraTextureSample_Windows.cpp
@@ -146,9 +146,8 @@ bool FElectraTextureSample::FinishInitialization()
 				}
 				case EElectraTextureSamplePixelEncoding::Y0CbY1Cr:
 				{
-					// TODO!!!!!!!! ("swapped" v216 - seems there is no real format for this?)
-					SampleFormat = EMediaTextureSampleFormat::Undefined;
-					return false;
+					SampleFormat = EMediaTextureSampleFormat::ShortYUY2;
+					break;
 				}
 				case EElectraTextureSamplePixelEncoding::YCbCr_Alpha:
 				{
diff --git a/Engine/Source/Runtime/Media/Public/IMediaTextureSample.h b/Engine/Source/Runtime/Media/Public/IMediaTextureSample.h
--- a/Engine/Source/Runtime/Media/Public/IMediaTextureSample.h
+++ b/Engine/Source/Runtime/Media/Public/IMediaTextureSample.h
@@ -87,6 +87,9 @@ enum class EMediaTextureSampleFormat
 	/** YUV v216 format which pack 2 pixel using 4 x 16bits components */
 	YUVv216,
 
+	/** Packed YUYV (Y0CbY1Cr) on 16-bit unsigned integer components. */
+	ShortYUY2,
+
 	/** 4:4:4:4 AY'CbCr 16-bit little endian full range alpha, video range Y'CbCr. */
 	Y416,
 
@@ -447,6 +450,7 @@ public:
 			case EMediaTextureSampleFormat::R4FL:		Off = MediaShaders::YUVOffsetFloat; break;
 			case EMediaTextureSampleFormat::Y416:
 			case EMediaTextureSampleFormat::P010:
+			case EMediaTextureSampleFormat::ShortYUY2:
 			case EMediaTextureSampleFormat::YUVv216:	Off = MediaShaders::YUVOffset16bits; break;
 			case EMediaTextureSampleFormat::YUVv210:	Off = MediaShaders::YUVOffset10bits; break;
 			default:									Off = MediaShaders::YUVOffset8bits; break;
diff --git a/Engine/Source/Runtime/MediaAssets/Private/Misc/MediaTextureResource.cpp b/Engine/Source/Runtime/MediaAssets/Private/Misc/MediaTextureResource.cpp
--- a/Engine/Source/Runtime/MediaAssets/Private/Misc/MediaTextureResource.cpp
+++ b/Engine/Source/Runtime/MediaAssets/Private/Misc/MediaTextureResource.cpp
@@ -237,6 +237,7 @@ namespace MediaTextureResourceHelpers
 			return PF_A2B10G10R10;
 
 		case EMediaTextureSampleFormat::YUVv216:
+		case EMediaTextureSampleFormat::ShortYUY2:
 			return PF_A16B16G16R16;
 
 		case EMediaTextureSampleFormat::Y416:
@@ -285,6 +286,7 @@ namespace MediaTextureResourceHelpers
 		case EMediaTextureSampleFormat::FloatRGB:
 		case EMediaTextureSampleFormat::P010:
 		case EMediaTextureSampleFormat::YUVv216:
+		case EMediaTextureSampleFormat::ShortYUY2:
 			return PF_FloatRGB;
 
 		case EMediaTextureSampleFormat::Y416:
@@ -1200,6 +1202,7 @@ void FMediaTextureResource::GetColorSpaceConversionMatrixForSample(const TShared
 				case EMediaTextureSampleFormat::CharYVYU:		// Y0CrY1Cb
 				case EMediaTextureSampleFormat::Char2VUY:		// CbY0CrY1
 				case EMediaTextureSampleFormat::YUVv216:		// CbY0CrY1
+				case EMediaTextureSampleFormat::ShortYUY2:		// Y0CbY1Cr
 				case EMediaTextureSampleFormat::CharUYVY:		// CbY0CrY1
 				{
 					auto YUVMtx = Sample->GetSampleToRGBMatrix();
@@ -1210,7 +1213,7 @@ void FMediaTextureResource::GetColorSpaceConversionMatrixForSample(const TShared
 					GraphicsPSOInit.BoundShaderState.PixelShaderRHI = ConvertShader.GetPixelShader();
 					SetGraphicsPipelineState(RHICmdList, GraphicsPSOInit, 0);
 					SetShaderParametersLegacyPS(RHICmdList, ConvertShader, InputTexture, OutputDim, YUVMtx, Sample->GetEncodingType(), ColorSpaceMtx,
-												SampleFormat!= EMediaTextureSampleFormat::CharYUY2 && SampleFormat != EMediaTextureSampleFormat::CharYVYU,	// Y or Cb first
+												SampleFormat!= EMediaTextureSampleFormat::CharYUY2 && SampleFormat != EMediaTextureSampleFormat::CharYVYU && SampleFormat != EMediaTextureSampleFormat::ShortYUY2,	// Y or Cb first
 												InputTexture->GetFormat() == PF_B8G8R8A8,																	// ARGB vs. ABGR (memory order)
 												SampleFormat == EMediaTextureSampleFormat::CharYVYU,														// Cb / Cr swap
 												Sample->GetToneMapMethod(),
