[MediaAssets] Changing the resource name to make it unique to help debugging the d3d and rhi validation error messages.

diff --git a/Engine/Source/Runtime/MediaAssets/Private/Misc/MediaTextureResource.cpp b/Engine/Source/Runtime/MediaAssets/Private/Misc/MediaTextureResource.cpp
--- a/Engine/Source/Runtime/MediaAssets/Private/Misc/MediaTextureResource.cpp
+++ b/Engine/Source/Runtime/MediaAssets/Private/Misc/MediaTextureResource.cpp
@@ -612,7 +612,7 @@ void FMediaTextureResource::Render(const FRenderParams& Params)
 							.SetFlags(IntermediateTextureDesc.Flags);
 
 						OutputTarget = RHICmdList.CreateTexture(Desc);
-						OutputTarget->SetName(TEXT("MediaTextureResourceOutput"));
+						OutputTarget->SetName(TEXT("MediaTextureColorConversionOutput"));
 						OutputTarget->SetOwnerName(GetOwnerName());
 
 						bRecreateOutputTarget = false;
