Media: Replace Bradford with DEFAULT_CHROMATIC_ADAPTATION_METHOD alias.

#Synced-CL 50608378

diff --git a/Engine/Source/Runtime/MediaAssets/Private/BaseMediaSourceColorSettings.cpp b/Engine/Source/Runtime/MediaAssets/Private/BaseMediaSourceColorSettings.cpp
--- a/Engine/Source/Runtime/MediaAssets/Private/BaseMediaSourceColorSettings.cpp
+++ b/Engine/Source/Runtime/MediaAssets/Private/BaseMediaSourceColorSettings.cpp
@@ -18,7 +18,7 @@ void FMediaSourceColorSettings::UpdateColorSpaceChromaticities()
 FNativeMediaSourceColorSettings::FNativeMediaSourceColorSettings()
 	: EncodingOverride(UE::Color::EEncoding::None)
 	, ColorSpaceOverride()
-	, ChromaticAdaptationMethod(UE::Color::EChromaticAdaptationMethod::Bradford)
+	, ChromaticAdaptationMethod(UE::Color::DEFAULT_CHROMATIC_ADAPTATION_METHOD)
 {}
 
 FNativeMediaSourceColorSettings::~FNativeMediaSourceColorSettings() = default;
