[MediaAssets] Media Texture: Fixing the incorrect number of mips clamping.

This affects the Tmv (Exr/Apv) code path when texture samples already have a full  mips chain.

The maximum number of mips calculation would clamp to a lesser number of mips, for instance 10 mips for a 1080x1080, instead of 11, leading to missing render target views when doing the conversion.

This was the cause of the following ensure:

LogOutputDevice: Error: === Handled ensure: === LogOutputDevice: Error: Ensure condition failed: RenderTargetView [File:D:\P4V\LCY-WD-20324_FNMAIN\Engine\Source\Runtime\D3D12RHI\Private\D3D12Commands.cpp] [Line: 889] LogOutputDevice: Error: Texture being set as render target has no RTV

/Game/Media/EXR_Seq/IMS_AtmosSmoke044 would do this once converted to APV.

#Synced-CL 51094493

diff --git a/Engine/Source/Runtime/MediaAssets/Private/Misc/MediaTextureResource.cpp b/Engine/Source/Runtime/MediaAssets/Private/Misc/MediaTextureResource.cpp
--- a/Engine/Source/Runtime/MediaAssets/Private/Misc/MediaTextureResource.cpp
+++ b/Engine/Source/Runtime/MediaAssets/Private/Misc/MediaTextureResource.cpp
@@ -51,11 +51,26 @@ static const FMatrix44f SplatMtx = { {1.0f, 0.0f, 0.0f, 0.0f}, {1.0f, 0.0f, 0.0f
 
 namespace UE::MediaTexture::Private
 {
+	/** Calculate the maximum number of mips for the given texture dimension in pixels. */
+	static int32 CalcMaxNumMips(const FIntPoint& InDim)
+	{
+		const int32 MaxDimension = FMath::Max(InDim.X, InDim.Y);
+
+		// FloorLog2NonZero requires a positive non-zero value.
+		if (MaxDimension <= 0)
+		{
+			return 1;
+		}
+
+		const uint32 NumMips = 1u + FMath::FloorLog2NonZero(static_cast<uint32>(MaxDimension));
+		return static_cast<int32>(NumMips);
+	}
+
 	uint8 ClampNumMips(uint8 InNumMips, const FIntPoint& InDim)
 	{
 		if (InNumMips > 1u)
 		{
-			const uint8 MaxNumMips = static_cast<uint8>(FMath::Min(GMaxTextureMipCount, FGenericPlatformMath::FloorToInt(FGenericPlatformMath::Log2(static_cast<float>(FGenericPlatformMath::Min(InDim.X, InDim.Y))))));
+			const uint8 MaxNumMips = static_cast<uint8>(FMath::Min(GMaxTextureMipCount, CalcMaxNumMips(InDim)));
 			InNumMips = FMath::Min(InNumMips, MaxNumMips);
 		}
 
