[ImgMedia] Fix OpenExr reader to read the correct number of mip levels.

For TMV playback, both APV (in a yuv444 format) and EXR (now with this fix) will use the full mip chain down to the 1x1 mip.
Several locations, in Exr Reader and MipMapInfo, where doing calculations on the mip size and number of tiles per mip level that now require carefully handling the edge cases to avoid zero size on one dimension because of rectangular mips.

MediaTextureResource:
Non-TMV readers (like DDS) could have a full mip chain that would lead to the same edge case that need to be fixed.

#Synced-CL 51094493

diff --git a/Engine/Plugins/Media/ImgMedia/Source/ExrReaderGpu/Private/ExrReaderGpu.cpp b/Engine/Plugins/Media/ImgMedia/Source/ExrReaderGpu/Private/ExrReaderGpu.cpp
--- a/Engine/Plugins/Media/ImgMedia/Source/ExrReaderGpu/Private/ExrReaderGpu.cpp
+++ b/Engine/Plugins/Media/ImgMedia/Source/ExrReaderGpu/Private/ExrReaderGpu.cpp
@@ -272,10 +272,10 @@ void FExrReader::CalculateTileOffsets
 	int64 CurrentPosition = 0;
 	for (int32 MipLevel = 0; MipLevel < NumMipLevels; MipLevel++)
 	{
-		int32 MipDiv = 1 << MipLevel;
+		int32 MipLevelDiv = 1 << MipLevel;
 
 		// Resolution of the texture in pixels for this mip level.
-		FIntPoint MipResolution = FullTextureResolution / MipDiv;
+		FIntPoint MipResolution = (FullTextureResolution / MipLevelDiv).ComponentMax(FIntPoint(1,1));
 
 		// Dimension of the texture in tiles including partial tiles
 		float TileResFractionX = float(MipResolution.X) / TileDimWithBorders.X;
diff --git a/Engine/Plugins/Media/ImgMedia/Source/ImgMedia/Private/Assets/ImgMediaMipMapInfo.cpp b/Engine/Plugins/Media/ImgMedia/Source/ImgMedia/Private/Assets/ImgMediaMipMapInfo.cpp
--- a/Engine/Plugins/Media/ImgMedia/Source/ImgMedia/Private/Assets/ImgMediaMipMapInfo.cpp
+++ b/Engine/Plugins/Media/ImgMedia/Source/ImgMedia/Private/Assets/ImgMediaMipMapInfo.cpp
@@ -55,8 +55,8 @@ FImgMediaTileSelection FImgMediaTileSelection::CreateForTargetMipLevel(const FIn
 	}
 
 	const int MipLevelDiv = 1 << TargetMipLevel;
-	int32 NumTilesX = FMath::Max(1, FMath::CeilToInt(float(MipZeroResolution.X / MipLevelDiv) / TileSize.X));
-	int32 NumTilesY = FMath::Max(1, FMath::CeilToInt(float(MipZeroResolution.Y / MipLevelDiv) / TileSize.Y));
+	int32 NumTilesX = FMath::Max(1, FMath::CeilToInt(float(MipZeroResolution.X) / MipLevelDiv / TileSize.X));
+	int32 NumTilesY = FMath::Max(1, FMath::CeilToInt(float(MipZeroResolution.Y) / MipLevelDiv / TileSize.Y));
 
 	return FImgMediaTileSelection(NumTilesX, NumTilesY, bDefaultVisibility);
 }
@@ -458,8 +458,8 @@ namespace {
 
 				if (InSequenceInfo.IsTiled())
 				{
-					CurrentNumTiles.X = FMath::CeilToInt(float(InSequenceInfo.Dim.X / MipLevelDiv) / InSequenceInfo.TilingDescription.TileSize.X);
-					CurrentNumTiles.Y = FMath::CeilToInt(float(InSequenceInfo.Dim.Y / MipLevelDiv) / InSequenceInfo.TilingDescription.TileSize.Y);
+					CurrentNumTiles.X = FMath::CeilToInt(float(InSequenceInfo.Dim.X) / MipLevelDiv / InSequenceInfo.TilingDescription.TileSize.X);
+					CurrentNumTiles.Y = FMath::CeilToInt(float(InSequenceInfo.Dim.Y) / MipLevelDiv / InSequenceInfo.TilingDescription.TileSize.Y);
 				}
 
 				// Starting with tiles at the highest mip level
diff --git a/Engine/Plugins/Media/ImgMedia/Source/ImgMedia/Private/Readers/ExrImgMediaReader.cpp b/Engine/Plugins/Media/ImgMedia/Source/ImgMedia/Private/Readers/ExrImgMediaReader.cpp
--- a/Engine/Plugins/Media/ImgMedia/Source/ImgMedia/Private/Readers/ExrImgMediaReader.cpp
+++ b/Engine/Plugins/Media/ImgMedia/Source/ImgMedia/Private/Readers/ExrImgMediaReader.cpp
@@ -67,11 +67,11 @@ FExrImgMediaReader::EReadResult FExrImgMediaReader::ReadTiles
 	FExrReader ChunkReader;
 	int MipLevelDiv = 1 << CurrentMipLevel;
 
-	FIntPoint MipResolution = ConverterParams.FullResolution / MipLevelDiv;
+	const FIntPoint MipResolution = (ConverterParams.FullResolution / MipLevelDiv).ComponentMax(FIntPoint(1,1));
 
 	FIntPoint DimensionInTiles
-		( FMath::CeilToInt(float(MipResolution.X) / ConverterParams.TileDimWithBorders.X)
-		, FMath::CeilToInt(float(MipResolution.Y) / ConverterParams.TileDimWithBorders.Y));
+		( FMath::CeilToInt(static_cast<float>(MipResolution.X) / ConverterParams.TileDimWithBorders.X)
+		, FMath::CeilToInt(static_cast<float>(MipResolution.Y) / ConverterParams.TileDimWithBorders.Y));
 
 	TArray<int32> NumTilesPerLevel;
 	int32 NumMipLevels = ConverterParams.bMipsInSeparateFiles ? 1 : ConverterParams.NumMipLevels;
@@ -392,12 +392,12 @@ bool FExrImgMediaReader::GetInfo(const FString& FilePath, FImgMediaFrameInfo& Ou
 		, FMath::CeilToInt(float(OutInfo.Dim.Y) / (OutInfo.TileDimensions.Y + OutInfo.TileBorder * 2))};
 		if (HeaderReader.ContainsMips())
 		{
-			OutInfo.NumMipLevels = HeaderReader.CalculateNumMipLevels(OutInfo.NumTiles);
+			OutInfo.NumMipLevels = HeaderReader.GetNumMipLevels();
 
 			SIZE_T SizeMip0 = OutInfo.UncompressedSize;
 			for (int32 Level = 1; Level < OutInfo.NumMipLevels; Level++)
 			{
-				OutInfo.UncompressedSize += SizeMip0 >> (2 * Level);
+				OutInfo.UncompressedSize += HeaderReader.GetUncompressedMipSize(Level);
 			}
 		}
 	}
diff --git a/Engine/Plugins/Media/ImgMedia/Source/ImgMedia/Private/Readers/ExrImgMediaReaderGpu.cpp b/Engine/Plugins/Media/ImgMedia/Source/ImgMedia/Private/Readers/ExrImgMediaReaderGpu.cpp
--- a/Engine/Plugins/Media/ImgMedia/Source/ImgMedia/Private/Readers/ExrImgMediaReaderGpu.cpp
+++ b/Engine/Plugins/Media/ImgMedia/Source/ImgMedia/Private/Readers/ExrImgMediaReaderGpu.cpp
@@ -122,10 +122,14 @@ FExrImgMediaReader::EReadResult FExrImgMediaReaderGpu::ReadMip
 	TRACE_CPUPROFILER_EVENT_SCOPE_TEXT(*FString::Printf(TEXT("ExrReaderGpu.ReadMip %d"), CurrentMipLevel));
 
 	// Next mip level.
-	int MipLevelDiv = 1 << CurrentMipLevel;
-	FIntPoint CurrentMipDim = ConverterParams.FullResolution / MipLevelDiv;
+	const int MipLevelDiv = 1 << CurrentMipLevel;
+	const FIntPoint CurrentMipDim = (ConverterParams.FullResolution / MipLevelDiv).ComponentMax(FIntPoint(1,1));
 	const FImgMediaFrameInfo& FrameInfo = ConverterParams.FrameInfo;
-	const SIZE_T BufferSize = GetBufferSize(CurrentMipDim, FrameInfo.NumChannels, FrameInfo.bHasTiles, FrameInfo.NumTiles / MipLevelDiv);
+	const FIntPoint CurrentNumTiles = FIntPoint(
+		FMath::Max(1, FMath::CeilToInt(static_cast<float>(FrameInfo.NumTiles.X) / MipLevelDiv)),
+		FMath::Max(1, FMath::CeilToInt(static_cast<float>(FrameInfo.NumTiles.Y) / MipLevelDiv)));
+
+	const SIZE_T BufferSize = GetBufferSize(CurrentMipDim, FrameInfo.NumChannels, FrameInfo.bHasTiles, CurrentNumTiles);
 	
 	FStructuredBufferPoolItemSharedPtr BufferData = SampleConverter->GetOrCreateMipLevelBuffer(
 		CurrentMipLevel,
@@ -337,7 +341,7 @@ bool FExrImgMediaReaderGpu::ReadFrame(int32 FrameId, const TMap<int32, FImgMedia
 			}
 
 			const int32 MipLevelDiv = 1 << CurrentMipLevel;
-			FIntPoint CurrentMipDim = ConverterParams.FullResolution / MipLevelDiv;
+			const FIntPoint CurrentMipDim = (ConverterParams.FullResolution / MipLevelDiv).ComponentMax(FIntPoint(1,1));
 
 			TArray<FIntRect>& Viewports = ConverterParams.Viewports.Add(CurrentMipLevel);
 			for (const FIntRect& TileRegion : CurrentTileSelection.GetVisibleRegions())
@@ -579,7 +583,7 @@ void FExrImgMediaReaderGpu::CreateSampleConverterCallback(TSharedPtr<FExrMediaTe
 		for (int32 MipLevel = 0; MipLevel <= MipToUpscale; MipLevel++)
 		{
 			int MipLevelDiv = 1 << MipLevel;
-			FIntPoint Dim = ConverterParams.FullResolution / MipLevelDiv;
+			const FIntPoint Dim = (ConverterParams.FullResolution / MipLevelDiv).ComponentMax(FIntPoint(1,1));
 
 			{
 				RHI_BREADCRUMB_EVENT_STAT(RHICmdList, ExrImgMediaReaderGpu_MipUpscale, "ExrImgGpu.MipRender.MipUpscale");
@@ -593,7 +597,7 @@ void FExrImgMediaReaderGpu::CreateSampleConverterCallback(TSharedPtr<FExrMediaTe
 				FStructuredBufferPoolItemSharedPtr BufferDataToUpscale = MipBuffers[MipToUpscale];
 				TArray<FIntRect> FakeViewport;
 				FakeViewport.Add(FIntRect(FIntPoint(0, 0), Dim));
-				RenderMip(RHICmdList, RenderTargetTextureRHI, ConverterParams, MipToUpscale, MipLevel, BufferDataToUpscale, (ConverterParams.FullResolution / (1 << MipToUpscale)), Dim, FakeViewport);
+				RenderMip(RHICmdList, RenderTargetTextureRHI, ConverterParams, MipToUpscale, MipLevel, BufferDataToUpscale, (ConverterParams.FullResolution / (1 << MipToUpscale)).ComponentMax(FIntPoint(1,1)), Dim, FakeViewport);
 			}
 		}
 
@@ -609,7 +613,7 @@ void FExrImgMediaReaderGpu::CreateSampleConverterCallback(TSharedPtr<FExrMediaTe
 
 			FStructuredBufferPoolItemSharedPtr BufferData = MipBuffers[MipLevel];
 			int MipLevelDiv = 1 << MipLevel;
-			FIntPoint Dim = ConverterParams.FullResolution / MipLevelDiv;
+			const FIntPoint Dim = (ConverterParams.FullResolution / MipLevelDiv).ComponentMax(FIntPoint(1,1));
 
 			if (BufferData.IsValid())
 			{
diff --git a/Engine/Plugins/Media/ImgMedia/Source/ImgMedia/Private/Readers/TmvImgMediaReader.cpp b/Engine/Plugins/Media/ImgMedia/Source/ImgMedia/Private/Readers/TmvImgMediaReader.cpp
--- a/Engine/Plugins/Media/ImgMedia/Source/ImgMedia/Private/Readers/TmvImgMediaReader.cpp
+++ b/Engine/Plugins/Media/ImgMedia/Source/ImgMedia/Private/Readers/TmvImgMediaReader.cpp
@@ -342,7 +342,7 @@ namespace UE::ImgMedia
 			}
 
 			const int32 MipLevelDiv = 1 << CurrentMipLevel;
-			FIntPoint CurrentMipDim = InConverterParams.FullResolution / MipLevelDiv;
+			FIntPoint CurrentMipDim = (InConverterParams.FullResolution / MipLevelDiv).ComponentMax(FIntPoint(1,1));
 			const FIntRect FullFrameViewport(FIntPoint::ZeroValue, CurrentMipDim);
 
 			// These are the coalesced visible regions for the current frame cumulative from past and current decode jobs.
diff --git a/Engine/Plugins/Media/ImgMedia/Source/OpenExrWrapper/Private/OpenExrWrapperContextReader.cpp b/Engine/Plugins/Media/ImgMedia/Source/OpenExrWrapper/Private/OpenExrWrapperContextReader.cpp
--- a/Engine/Plugins/Media/ImgMedia/Source/OpenExrWrapper/Private/OpenExrWrapperContextReader.cpp
+++ b/Engine/Plugins/Media/ImgMedia/Source/OpenExrWrapper/Private/OpenExrWrapperContextReader.cpp
@@ -127,6 +127,14 @@ int64 FOpenExrHeaderReader::GetUncompressedSize() const
 	return (static_cast<int64>(Window.X) * static_cast<int64>(Window.Y) * static_cast<int64>(TexelSize));
 }
 
+int64 FOpenExrHeaderReader::GetUncompressedMipSize(int32 InMipLevel) const
+{
+	const int32 TexelSize = GetPixelSize();
+	const int32 MipLevelDiv = 1 << InMipLevel;
+	const FIntPoint MipWindow = (GetDataWindow() / MipLevelDiv).ComponentMax(FIntPoint(1,1));
+	return (static_cast<int64>(MipWindow.X) * static_cast<int64>(MipWindow.Y) * static_cast<int64>(TexelSize));
+}
+
 int32 FOpenExrHeaderReader::GetNumChannels() const
 {
 	const exr_attribute_t* Attribute = nullptr;
@@ -173,6 +181,17 @@ bool FOpenExrHeaderReader::ContainsMips() const
 	return false;
 }
 
+int32 FOpenExrHeaderReader::GetNumMipLevels() const
+{
+	int32_t NumLevelsX = 0, NumLevelsY = 0;
+	exr_result_t Result = exr_get_tile_levels(*((exr_context_t*)FileContext.Get()), 0, &NumLevelsX, &NumLevelsY);
+	if (Result == EXR_ERR_SUCCESS)
+	{
+		return FMath::Min(NumLevelsX, NumLevelsY);
+	}
+	return 1;
+}
+
 int32 FOpenExrHeaderReader::CalculateNumMipLevels(const FIntPoint& NumTiles) const
 {
 	const exr_attribute_t* Attribute = nullptr;
diff --git a/Engine/Plugins/Media/ImgMedia/Source/OpenExrWrapper/Public/OpenExrWrapper.h b/Engine/Plugins/Media/ImgMedia/Source/OpenExrWrapper/Public/OpenExrWrapper.h
--- a/Engine/Plugins/Media/ImgMedia/Source/OpenExrWrapper/Public/OpenExrWrapper.h
+++ b/Engine/Plugins/Media/ImgMedia/Source/OpenExrWrapper/Public/OpenExrWrapper.h
@@ -73,9 +73,11 @@ public:
 	UE_API FIntPoint GetDataWindow() const;
 	UE_API FFrameRate GetFrameRate(const FFrameRate& DefaultValue) const;
 	UE_API int64 GetUncompressedSize() const;
+	UE_API int64 GetUncompressedMipSize(int32 InMipLevel) const;
 	UE_API int32 GetNumChannels() const;
 	UE_API int32 GetPixelSize() const;
 	UE_API bool ContainsMips() const;
+	UE_API int32 GetNumMipLevels() const;
 	UE_API int32 CalculateNumMipLevels(const FIntPoint& NumTiles) const;
 
 	/** Determines if Gpu reader can be utilized. */
diff --git a/Engine/Source/Runtime/MediaAssets/Private/Misc/MediaTextureResource.cpp b/Engine/Source/Runtime/MediaAssets/Private/Misc/MediaTextureResource.cpp
--- a/Engine/Source/Runtime/MediaAssets/Private/Misc/MediaTextureResource.cpp
+++ b/Engine/Source/Runtime/MediaAssets/Private/Misc/MediaTextureResource.cpp
@@ -333,6 +333,57 @@ namespace MediaTextureResourceHelpers
 		}
 	}
 
+	/**
+	 * Helper function to calculate the stride of the sample buffer
+	 * using mip0's stride up to the point where there is 1 pixel left.
+	 */
+	uint32 CalcSampleStrideForMipLevel(const TSharedPtr<IMediaTextureSample, ESPMode::ThreadSafe>& InSample, uint8 InMipLevel)
+	{
+		const uint32 MipStride = InSample->GetStride() >> InMipLevel;
+		if (MipStride > 0)
+		{
+			return MipStride;
+		}
+		
+		// Calculate the smallest stride for this format (for 1 pixel).
+		const EPixelFormat PixelFormat = GetPixelFormat(InSample);
+		const FPixelFormatInfo& PixelFormatInfo = GPixelFormats[PixelFormat];
+		return PixelFormatInfo.GetBlockCountForWidth(1) * PixelFormatInfo.BlockBytes;
+	}
+
+	/**
+	 * Updates the destination RHI texture from the sample's memory buffers. 
+	 */
+	void UpdateTextureFromSample(
+		FRHICommandListImmediate& InRhiCmdList,
+		const TSharedPtr<IMediaTextureSample, ESPMode::ThreadSafe>& InSample,
+		uint8 InNumMips,
+		const TRefCountPtr<FRHITexture>& InDestTexture)
+	{
+		const uint8* Data = (const uint8*)InSample->GetBuffer();
+		check(Data);
+		const FPixelFormatInfo& FormatInfo = GPixelFormats[GetPixelFormat(InSample)];
+		const FIntPoint SampleDim = InSample->GetDim();
+
+		for (uint8 MipLevel = 0; MipLevel < InNumMips; ++MipLevel)
+		{
+			const uint32 Width = FMath::Max(1, SampleDim.X >> MipLevel);
+			const uint32 Height = FMath::Max(1, SampleDim.Y >> MipLevel);
+			const uint32 Stride = CalcSampleStrideForMipLevel(InSample, MipLevel);
+
+			// Don't let a zero stride go through rhi, it causes a delayed error that is hard to track.
+			if (!ensure(Stride != 0))
+			{
+				break;
+			}
+
+			FUpdateTextureRegion2D Region(0, 0, 0, 0, Width, Height);
+			InRhiCmdList.UpdateTexture2D(InDestTexture, MipLevel, Region, Stride, Data);
+
+			const uint32 HeightUnit = (Height + FormatInfo.BlockSizeY - 1) / FormatInfo.BlockSizeY;
+			Data += Stride * HeightUnit;
+		}
+	}
 
 	bool RequiresSrgbTexture(const TSharedPtr<IMediaTextureSample, ESPMode::ThreadSafe>& Sample)
 	{
@@ -1045,16 +1096,7 @@ void FMediaTextureResource::ConvertSample(FRHICommandListImmediate& RHICmdList,
 			}
 
 			// copy sample data to input render target
-			const uint8* Data = (const uint8*)Sample->GetBuffer();
-			check(Data);
-			for (uint8 MipLevel = 0; MipLevel < SampleNumMips; ++MipLevel)
-			{
-				uint32 Stride = Sample->GetStride() >> MipLevel;
-				uint32 Height = SampleDim.Y >> MipLevel;
-				FUpdateTextureRegion2D Region(0, 0, 0, 0, SampleDim.X >> MipLevel, Height);
-				RHICmdList.UpdateTexture2D(InputTarget, MipLevel, Region, Stride, Data);
-				Data += Stride * Height;
-			}
+			MediaTextureResourceHelpers::UpdateTextureFromSample(RHICmdList, Sample, SampleNumMips, InputTarget);
 
 			InputTexture = InputTarget;
 		}
@@ -1505,22 +1547,7 @@ void FMediaTextureResource::CopySample(FRHICommandListImmediate& RHICmdList, con
 			if (!Converter || !Converter->Convert(RHICmdList, RenderTargetTextureRHI, Hints))
 			{
 				// Copy sample data (from CPU mem) to output render target
-				const FIntPoint SampleDim = Sample->GetDim();
-				const uint8* Data = (const uint8*)Sample->GetBuffer();
-				const FPixelFormatInfo& FormatInfo = GPixelFormats[MediaTextureResourceHelpers::GetPixelFormat(Sample)];
-
-				for (uint8 MipLevel = 0; MipLevel < SampleNumMips; ++MipLevel)
-				{
-					uint32 Width = SampleDim.X >> MipLevel;
-					uint32 Height = SampleDim.Y >> MipLevel;
-					uint32 Stride = Sample->GetStride() >> MipLevel;
-
-					FUpdateTextureRegion2D Region(0, 0, 0, 0, Width, Height);
-					RHICmdList.UpdateTexture2D(RenderTargetTextureRHI, MipLevel, Region, Stride, Data);
-
-					uint32 HeightUnit = (Height + FormatInfo.BlockSizeY - 1) / FormatInfo.BlockSizeY;
-					Data += Stride * HeightUnit;
-				}
+				MediaTextureResourceHelpers::UpdateTextureFromSample(RHICmdList, Sample, SampleNumMips, RenderTargetTextureRHI);
 			}
 
 			// Make sure resource is in SRV mode again
