[MediaAssets] MediaTextures will automatically recache affected material's uniform expressions upon dimensions change.

MediaTextures can change dimensions and format dynamically as part of media playback.

#Synced-CL 52176027

diff --git a/Engine/Source/Runtime/MediaAssets/Private/Assets/MediaTexture.cpp b/Engine/Source/Runtime/MediaAssets/Private/Assets/MediaTexture.cpp
--- a/Engine/Source/Runtime/MediaAssets/Private/Assets/MediaTexture.cpp
+++ b/Engine/Source/Runtime/MediaAssets/Private/Assets/MediaTexture.cpp
@@ -1,25 +1,25 @@
 // Copyright Epic Games, Inc. All Rights Reserved.
 
 #include "MediaTexture.h"
-#include "MediaAssetsPrivate.h"
 
 #include "ExternalTexture.h"
 #include "IMediaClock.h"
 #include "IMediaClockSink.h"
 #include "IMediaModule.h"
+#include "IMediaPlayer.h"
+#include "IMediaTextureSample.h"
+#include "Materials/MaterialInterface.h"
+#include "MediaAssetsPrivate.h"
+#include "MediaPlayer.h"
 #include "MediaPlayerFacade.h"
+#include "Misc/MediaTextureResource.h"
 #include "Modules/ModuleManager.h"
+#include "ObjectCacheContext.h"
+#include "RectLightTexture.h"
 #include "RenderUtils.h"
 #include "RenderingThread.h"
 #include "UObject/WeakObjectPtrTemplates.h"
 
-#include "MediaPlayer.h"
-#include "IMediaPlayer.h"
-#include "Misc/MediaTextureResource.h"
-#include "IMediaTextureSample.h"
-
-#include "RectLightTexture.h"
-
 #include UE_INLINE_GENERATED_CPP_BY_NAME(MediaTexture)
 
 /* Local helpers
@@ -87,6 +87,7 @@ UMediaTexture::UMediaTexture(const FObjectInitializer& ObjectInitializer)
 	, CurrentOrientation(MTORI_Original)
 	, DefaultGuid(FGuid::NewGuid())
 	, Dimensions(FIntPoint::ZeroValue)
+	, LastDimensions(FIntPoint::ZeroValue)
 	, bIsCleared(false)
 	, Size(0)
 	, CachedNextSampleTime(FTimespan::MinValue())
@@ -372,6 +373,17 @@ void UMediaTexture::TickResource(FTimespan Timecode)
 
 	// media player bookkeeping
 	UpdatePlayerAndQueue();
+	
+	// Detect texture dimension changes
+	if (LastDimensions != Dimensions)
+	{
+		FObjectCacheContextScope ObjectCache;
+		for (UMaterialInterface* MaterialInterface : ObjectCache.GetContext().GetMaterialsAffectedByTexture(this))
+		{
+			MaterialInterface->RecacheUniformExpressions(false);
+		}
+		LastDimensions = Dimensions;
+	}
 
 	if (!CurrentPlayer.IsValid())
 	{
diff --git a/Engine/Source/Runtime/MediaAssets/Public/MediaTexture.h b/Engine/Source/Runtime/MediaAssets/Public/MediaTexture.h
--- a/Engine/Source/Runtime/MediaAssets/Public/MediaTexture.h
+++ b/Engine/Source/Runtime/MediaAssets/Public/MediaTexture.h
@@ -331,6 +331,9 @@ private:
 	/** Current width and height of the resource (in pixels). */
 	FIntPoint Dimensions;
 
+	/** Last Tick's width and height of the resource (in pixels). */
+	FIntPoint LastDimensions;
+
 	/** The previously clear color. */
 	FLinearColor LastClearColor;
 
